- Updated the tooltip as proposed by [REDACTED_PERSON] (in a previous review)

diff --git a/Engine/Plugins/Developer/Concert/DisasterRecoveryClient/Source/DisasterRecoveryClient/Private/SDisasterRecovery.cpp b/Engine/Plugins/Developer/Concert/DisasterRecoveryClient/Source/DisasterRecoveryClient/Private/SDisasterRecovery.cpp
--- a/Engine/Plugins/Developer/Concert/DisasterRecoveryClient/Source/DisasterRecoveryClient/Private/SDisasterRecovery.cpp
+++ b/Engine/Plugins/Developer/Concert/DisasterRecoveryClient/Source/DisasterRecoveryClient/Private/SDisasterRecovery.cpp
@@ -193,7 +193,7 @@ void SDisasterRecovery::Construct(const FArguments& InArgs, TArray<TSharedPtr<FC
 					SNew(SButton)
 					.ForegroundColor(FLinearColor::White)
 					.ButtonStyle(FEditorStyle::Get(), TEXT("FlatButton.Success"))
-					.ToolTipText(LOCTEXT("RecoverTooltip", "Recover your assets to the state they were up to the selected recovery point (included)"))
+					.ToolTipText(LOCTEXT("RecoverTooltip", "Restore your assets to their state at the selected recovery point (inclusive)"))
 					.OnClicked(this, &SDisasterRecovery::OnRecoverClicked)
 					.HAlign(HAlign_Center)
 					.ContentPadding(FMargin(14, 3))
@@ -210,7 +210,7 @@ void SDisasterRecovery::Construct(const FArguments& InArgs, TArray<TSharedPtr<FC
 					SNew(SButton)
 					.ForegroundColor(FLinearColor::White)
 					.ButtonStyle(FEditorStyle::Get(), TEXT("FlatButton.Danger"))
-					.ToolTipText(LOCTEXT("CancelRecoveryTooltip", "Don't recover, resume work using the assets as they are on disk (last saved), losing unsaved work."))
+					.ToolTipText(LOCTEXT("CancelRecoveryTooltip", "Discard any recoverable data for your assets and continue with their last saved state"))
 					.OnClicked(this, &SDisasterRecovery::OnCancelRecoveryClicked)
 					.HAlign(HAlign_Center)
 					.ContentPadding(FMargin(14, 3))
