- A line of code used for debugging only was left uncommented.

diff --git a/Engine/Plugins/Developer/Concert/DisasterRecoveryClient/Source/DisasterRecoveryClient/Private/DisasterRecoveryClientModule.cpp b/Engine/Plugins/Developer/Concert/DisasterRecoveryClient/Source/DisasterRecoveryClient/Private/DisasterRecoveryClientModule.cpp
--- a/Engine/Plugins/Developer/Concert/DisasterRecoveryClient/Source/DisasterRecoveryClient/Private/DisasterRecoveryClientModule.cpp
+++ b/Engine/Plugins/Developer/Concert/DisasterRecoveryClient/Source/DisasterRecoveryClient/Private/DisasterRecoveryClientModule.cpp
@@ -247,7 +247,7 @@ private:
 		// Normally, bAutoRestoreLastSession is set true here and overwritten to false when the app exits normally, but when running under the debugger, auto-restore is disabled as
 		// programmers kill applications (stop the debugger) and this should not count as a crash (unless you want to simulate crash this way - see below).
 		SessionInfoToSave.bAutoRestoreLastSession = !FPlatformMisc::IsDebuggerPresent();
-		SessionInfoToSave.bAutoRestoreLastSession = true; // <- THIS CODE MUST BE COMMENTED BEFORE SUBMITTING: Uncomment for debugging purpose only. Simulate a crash by stopping the debugger during a session.
+		//SessionInfoToSave.bAutoRestoreLastSession = true; // <- THIS CODE MUST BE COMMENTED BEFORE SUBMITTING: Uncomment for debugging purpose only. Simulate a crash by stopping the debugger during a session.
 		SaveDisasterRecoverySessionInfo(SessionInfoToSave);
 	}
 
