- Disable Auto-Save recovery when Disaster Recovery is enabled and a session to recover is available.

diff --git a/Engine/Plugins/Developer/Concert/DisasterRecoveryClient/Source/DisasterRecoveryClient/DisasterRecoveryClient.Build.cs b/Engine/Plugins/Developer/Concert/DisasterRecoveryClient/Source/DisasterRecoveryClient/DisasterRecoveryClient.Build.cs
--- a/Engine/Plugins/Developer/Concert/DisasterRecoveryClient/Source/DisasterRecoveryClient/DisasterRecoveryClient.Build.cs
+++ b/Engine/Plugins/Developer/Concert/DisasterRecoveryClient/Source/DisasterRecoveryClient/DisasterRecoveryClient.Build.cs
@@ -21,7 +21,8 @@ namespace UnrealBuildTool.Rules
 					"ConcertSyncClient",
 					"ConcertSyncCore",
 					"ConcertSyncUI",
-					"ConcertTransport",					
+					"ConcertTransport",
+					"UnrealEd",
 					"EditorStyle",
 					"InputCore",
 					"Serialization",
diff --git a/Engine/Plugins/Developer/Concert/DisasterRecoveryClient/Source/DisasterRecoveryClient/Private/DisasterRecoveryClientModule.cpp b/Engine/Plugins/Developer/Concert/DisasterRecoveryClient/Source/DisasterRecoveryClient/Private/DisasterRecoveryClientModule.cpp
--- a/Engine/Plugins/Developer/Concert/DisasterRecoveryClient/Source/DisasterRecoveryClient/Private/DisasterRecoveryClientModule.cpp
+++ b/Engine/Plugins/Developer/Concert/DisasterRecoveryClient/Source/DisasterRecoveryClient/Private/DisasterRecoveryClientModule.cpp
@@ -20,6 +20,9 @@
 #include "Backends/JsonStructDeserializerBackend.h"
 #include "DisasterRecoveryFSM.h"
 #include "DisasterRecoverySettings.h"
+#include "UnrealEdGlobals.h"
+#include "Editor/UnrealEdEngine.h"
+#include "IPackageAutoSaver.h"
 
 #if WITH_EDITOR
 	#include "ISettingsModule.h"
@@ -311,6 +314,12 @@ private:
 		// Starts the recovery process.
 		if (bRecoveringFromCrash)
 		{
+			if (GUnrealEd)
+			{
+				// Prevent the "Auto-Save" system from restoring the packages before Disaster Recovery plugin.
+				GUnrealEd->GetPackageAutoSaver().DisableRestorePromptAndDeclinePackageRecovery();
+			}
+
 			DisasterRecoveryUtil::StartRecovery(DisasterRecoveryClient->GetConcertClient(), InSessionInfoToRestore->LastSessionName, /*bLiveDataOnly*/ false);
 		}
 
diff --git a/Engine/Source/Editor/UnrealEd/Private/PackageAutoSaver.cpp b/Engine/Source/Editor/UnrealEd/Private/PackageAutoSaver.cpp
--- a/Engine/Source/Editor/UnrealEd/Private/PackageAutoSaver.cpp
+++ b/Engine/Source/Editor/UnrealEd/Private/PackageAutoSaver.cpp
@@ -290,7 +290,7 @@ void FPackageAutoSaver::OfferToRestorePackages()
 {
 	bool bRemoveRestoreFile = true;
 
-	if(HasPackagesToRestore())
+	if(HasPackagesToRestore() && !bAutoDeclineRecovery) // if bAutoDeclineRecovery is true, do like the user selected to decline. (then remove the restore files)
 	{
 		// If we failed to restore, keep the restore information around
 		if(PackageRestore::PromptToRestorePackages(PackagesThatCanBeRestored) == FEditorFileUtils::PR_Failure)
diff --git a/Engine/Source/Editor/UnrealEd/Private/PackageAutoSaver.h b/Engine/Source/Editor/UnrealEd/Private/PackageAutoSaver.h
--- a/Engine/Source/Editor/UnrealEd/Private/PackageAutoSaver.h
+++ b/Engine/Source/Editor/UnrealEd/Private/PackageAutoSaver.h
@@ -36,10 +36,8 @@ public:
 	virtual bool HasPackagesToRestore() const override;
 	virtual void OfferToRestorePackages() override;
 	virtual void OnPackagesDeleted(const TArray<UPackage*>& DeletedPackages) override;
-	virtual bool IsAutoSaving() const override
-	{
-		return bIsAutoSaving;
-	}
+	virtual bool IsAutoSaving() const override { return bIsAutoSaving; }
+	virtual void DisableRestorePromptAndDeclinePackageRecovery() override { bAutoDeclineRecovery = true; }
 
 private:
 	/**
@@ -124,6 +122,9 @@ private:
 	/** If we are delaying the time a little bit because we failed to save */
 	bool bDelayingDueToFailedSave;
 
+	/** Flag as true if recovery prompt is disabled. Should be true if another system is restoring, replacing the auto-saver recovery, like the Disaster Recovery plugin. */
+	bool bAutoDeclineRecovery = false;
+
 	/** Used to reference to the active auto-save warning notification */
 	TWeakPtr<SNotificationItem> AutoSaveNotificationPtr;
 
diff --git a/Engine/Source/Editor/UnrealEd/Public/IPackageAutoSaver.h b/Engine/Source/Editor/UnrealEd/Public/IPackageAutoSaver.h
--- a/Engine/Source/Editor/UnrealEd/Public/IPackageAutoSaver.h
+++ b/Engine/Source/Editor/UnrealEd/Public/IPackageAutoSaver.h
@@ -60,4 +60,12 @@ public:
 
 	/** Called when packages are deleted in the editor */
 	virtual void OnPackagesDeleted(const TArray<UPackage*>& DeletedPackages) = 0;
+
+	/**
+	 * Set a flag to bypass the recovery UI prompt and to automatically decline package recovery. The auto-saver
+	 * flow runs as usual but when the user is normally prompted for recovery, the system will bypass the prompt and select
+	 * to dismiss automatically if the flag it set. The flag must be raised before the Engine reach the recovery point
+	 * during the boot process, otherwise, it has no effect.
+	 */
+	virtual void DisableRestorePromptAndDeclinePackageRecovery() = 0;
 };
