- Added missing call to initialize the concert sytle set from the disaster recover module (was initialized in multi-user module, so if both were enabled, it did not crash)

diff --git a/Engine/Plugins/Developer/Concert/DisasterRecoveryClient/Source/DisasterRecoveryClient/Private/DisasterRecoveryClientModule.cpp b/Engine/Plugins/Developer/Concert/DisasterRecoveryClient/Source/DisasterRecoveryClient/Private/DisasterRecoveryClientModule.cpp
--- a/Engine/Plugins/Developer/Concert/DisasterRecoveryClient/Source/DisasterRecoveryClient/Private/DisasterRecoveryClientModule.cpp
+++ b/Engine/Plugins/Developer/Concert/DisasterRecoveryClient/Source/DisasterRecoveryClient/Private/DisasterRecoveryClientModule.cpp
@@ -8,6 +8,7 @@
 #include "IConcertClient.h"
 #include "IConcertSession.h"
 #include "IConcertSyncClient.h"
+#include "ConcertFrontendStyle.h"
 
 #include "Misc/Paths.h"
 #include "Misc/CoreDelegates.h"
@@ -46,6 +47,9 @@ public:
 		// Wait for init to finish before starting the Disaster Recovery service
 		FCoreDelegates::OnFEngineLoopInitComplete.AddRaw(this, &FDisasterRecoveryClientModule::OnEngineInitComplete);
 
+		// Initialize Style
+		FConcertFrontendStyle::Initialize();
+
 		// Register the Disaster Recovery Settings panel.
 		RegisterSettings();
 	}
