- Changed the DisasterRecoveryClient module type from "Editor" to "EditorNoCommandlet"

diff --git a/Engine/Plugins/Developer/Concert/DisasterRecoveryClient/DisasterRecoveryClient.uplugin b/Engine/Plugins/Developer/Concert/DisasterRecoveryClient/DisasterRecoveryClient.uplugin
--- a/Engine/Plugins/Developer/Concert/DisasterRecoveryClient/DisasterRecoveryClient.uplugin
+++ b/Engine/Plugins/Developer/Concert/DisasterRecoveryClient/DisasterRecoveryClient.uplugin
@@ -17,7 +17,7 @@
 	"Modules": [
 		{
 			"Name": "DisasterRecoveryClient",
-			"Type": "Editor",
+			"Type": "EditorNoCommandlet",
 			"LoadingPhase": "Default"
 		}
 	],
