- Persisted all restored session changes in the game content folder.

diff --git a/Engine/Plugins/Developer/Concert/DisasterRecoveryClient/Source/DisasterRecoveryClient/Private/DisasterRecoveryClientModule.cpp b/Engine/Plugins/Developer/Concert/DisasterRecoveryClient/Source/DisasterRecoveryClient/Private/DisasterRecoveryClientModule.cpp
--- a/Engine/Plugins/Developer/Concert/DisasterRecoveryClient/Source/DisasterRecoveryClient/Private/DisasterRecoveryClientModule.cpp
+++ b/Engine/Plugins/Developer/Concert/DisasterRecoveryClient/Source/DisasterRecoveryClient/Private/DisasterRecoveryClientModule.cpp
@@ -324,7 +324,7 @@ private:
 				GUnrealEd->GetPackageAutoSaver().DisableRestorePromptAndDeclinePackageRecovery();
 			}
 
-			DisasterRecoveryUtil::StartRecovery(DisasterRecoveryClient->GetConcertClient(), InSessionInfoToRestore->LastSessionName, /*bLiveDataOnly*/ false);
+			DisasterRecoveryUtil::StartRecovery(DisasterRecoveryClient.ToSharedRef(), InSessionInfoToRestore->LastSessionName, /*bLiveDataOnly*/ false);
 		}
 
 		return true;
diff --git a/Engine/Plugins/Developer/Concert/DisasterRecoveryClient/Source/DisasterRecoveryClient/Private/DisasterRecoveryFSM.cpp b/Engine/Plugins/Developer/Concert/DisasterRecoveryClient/Source/DisasterRecoveryClient/Private/DisasterRecoveryFSM.cpp
--- a/Engine/Plugins/Developer/Concert/DisasterRecoveryClient/Source/DisasterRecoveryClient/Private/DisasterRecoveryFSM.cpp
+++ b/Engine/Plugins/Developer/Concert/DisasterRecoveryClient/Source/DisasterRecoveryClient/Private/DisasterRecoveryFSM.cpp
@@ -28,7 +28,7 @@ public:
 	static constexpr int64 MaxActivityPerRequest = 1024;
 
 public:
-	FDisasterRecoveryFSM(IConcertClientRef InClient, const FString& InSessionNameToRecover, bool bInLiveDataOnly);
+	FDisasterRecoveryFSM(TSharedRef<IConcertSyncClient> InSyncClient, const FString& InSessionNameToRecover, bool bInLiveDataOnly);
 	~FDisasterRecoveryFSM();
 
 private:
@@ -61,16 +61,19 @@ private:
 	/** Displays the session activities and gather the user selection. Transit to the next state (restoring the session) once the recover windows is closed by the user. */
 	void DisplayRecoveryUI();
 
-	/** Create and join a new recovery session is there is nothing to restore (or user selected to not restore), then go to the synchronization session state state. */
+	/** Creates and join a new recovery session is there is nothing to restore (or user selected to not restore), then go to the synchronization session state state. */
 	void CreateAndJoinSession();
 
-	/** Restore and join the session on the server at the selected point in time, then go to the synchronization state. */
+	/** Restores and join the session on the server at the selected point in time, then go to the synchronization state. */
 	void RestoreAndJoinSession();
 
+	/** Persists the recovered transactions locally, applying all changes to the game content folder. */
+	void PersistRecoveredChanges();
+
 	/** Displays the error that halted the recovery FSM and go to the exit state. */
 	void DisplayError();
 
-	/** Request a transition to the next state, applied to the next Tick(). Used when triggering a transition in a TFuture continuation. */
+	/** Requests a transition to the next state, applied to the next Tick(). Used when triggering a transition in a TFuture continuation. */
 	void RequestTransitTo(FState& NextState)
 	{
 		NextStatePending = &NextState;
@@ -85,6 +88,7 @@ private:
 	}
 
 private:
+	TSharedRef<IConcertSyncClient> SyncClient;
 	IConcertClientRef Client;
 	FDelegateHandle TickerHandle;
 
@@ -108,6 +112,7 @@ private:
 	FState RestoreAndJoinSessionState; // Restore the session on the server.
 	FState CreateAndJoinSessionState;  // Create a new session on the server.
 	FState SynchronizeState;           // Recover the assets to the selected point.
+	FState PersistChangesState;        // Applies the recovered transactions to the game content directory packages (making them dirty).
 	FState ErrorState;                 // Displays the error and halt the FSM.
 	FState ExitState;                  // State to exit the FSM.
 	FState* CurrentState = nullptr;    // The current state.
@@ -117,8 +122,9 @@ private:
 static TSharedPtr<FDisasterRecoveryFSM> RecoveryFSM;
 
 
-FDisasterRecoveryFSM::FDisasterRecoveryFSM(IConcertClientRef InClient, const FString& SessionNameToRecover, bool bInLiveDataOnly)
-	: Client(InClient)
+FDisasterRecoveryFSM::FDisasterRecoveryFSM(TSharedRef<IConcertSyncClient> InSyncClient, const FString& SessionNameToRecover, bool bInLiveDataOnly)
+	: SyncClient(InSyncClient)
+	, Client(InSyncClient->GetConcertClient())
 	, RecoverySessionName(SessionNameToRecover)
 	, bLiveDataOnly(bInLiveDataOnly)
 	, ExitMessage(TEXT("Disaster recovery process completed successfully."))
@@ -154,6 +160,7 @@ FDisasterRecoveryFSM::FDisasterRecoveryFSM(IConcertClientRef InClient, const FSt
 	RestoreAndJoinSessionState.OnExit  = [this]() { Client->OnSessionConnectionChanged().RemoveAll(this); };
 	SynchronizeState.OnEnter           = [this]() { IDisasterRecoveryClientModule::Get().GetClient()->GetWorkspace()->OnWorkspaceSynchronized().AddRaw(this, &FDisasterRecoveryFSM::OnWorkspaceSynchronized); };
 	SynchronizeState.OnExit            = [this]() { IDisasterRecoveryClientModule::Get().GetClient()->GetWorkspace()->OnWorkspaceSynchronized().RemoveAll(this); };
+	PersistChangesState.OnTick         = [this]() { PersistRecoveredChanges(); };
 	ErrorState.OnEnter                 = [this]() { NextStatePending = nullptr; DisplayError(); };
 	ExitState.OnEnter                  = [=]()    { Terminate(); };
 
@@ -188,7 +195,7 @@ void FDisasterRecoveryFSM::OnSessionConnectionChanged(IConcertClientSession& Ses
 void FDisasterRecoveryFSM::OnWorkspaceSynchronized()
 {
 	check(CurrentState == &SynchronizeState);
-	TransitTo(ExitState); // Disaster recovery process completed successfully.
+	TransitTo(PersistChangesState);
 }
 
 void FDisasterRecoveryFSM::LookupRecoveryServer()
@@ -372,6 +379,13 @@ void FDisasterRecoveryFSM::RestoreAndJoinSession()
 	});
 }
 
+void FDisasterRecoveryFSM::PersistRecoveredChanges()
+{
+	// Save live transactions to package, gather files changed in the Concert sandbox and apply the changes to the content directory.
+	SyncClient->PersistAllSessionChanges();
+	TransitTo(ExitState); // Disaster recovery process completed successfully.
+}
+
 void FDisasterRecoveryFSM::DisplayError()
 {
 	FAsyncTaskNotificationConfig NotificationConfig;
@@ -407,11 +421,11 @@ bool FDisasterRecoveryFSM::Tick(float)
 	return true;
 }
 
-void DisasterRecoveryUtil::StartRecovery(IConcertClientRef Client, const FString& SessionNameToRecover, bool bLiveDataOnly)
+void DisasterRecoveryUtil::StartRecovery(TSharedRef<IConcertSyncClient> SyncClient, const FString& SessionNameToRecover, bool bLiveDataOnly)
 {
 	if (!RecoveryFSM)
 	{
-		RecoveryFSM = MakeShared<FDisasterRecoveryFSM>(Client, SessionNameToRecover, bLiveDataOnly);
+		RecoveryFSM = MakeShared<FDisasterRecoveryFSM>(SyncClient, SessionNameToRecover, bLiveDataOnly);
 	}
 }
 
diff --git a/Engine/Plugins/Developer/Concert/DisasterRecoveryClient/Source/DisasterRecoveryClient/Private/DisasterRecoveryFSM.h b/Engine/Plugins/Developer/Concert/DisasterRecoveryClient/Source/DisasterRecoveryClient/Private/DisasterRecoveryFSM.h
--- a/Engine/Plugins/Developer/Concert/DisasterRecoveryClient/Source/DisasterRecoveryClient/Private/DisasterRecoveryFSM.h
+++ b/Engine/Plugins/Developer/Concert/DisasterRecoveryClient/Source/DisasterRecoveryClient/Private/DisasterRecoveryFSM.h
@@ -5,16 +5,18 @@
 #include "CoreMinimal.h"
 #include "IConcertModule.h"
 
+class IConcertSyncClient;
+
 namespace DisasterRecoveryUtil
 {
 
 /**
  * Start the recovery flow. This should be called when the previous instance of the editor exited unexpectedly.
- * @param Client The client configured with the recovery server URL, the recovery session name and recovery archive name.
+ * @param SyncClient The client configured with the recovery server URL, the recovery session name and recovery archive name.
  * @param SessionNameToRecover The name of the session to recover.
  * @param bLiveDataOnly Filter the recovery data to only recover live transaction data. (Transaction that were not saved to disk yet)
  */
-void StartRecovery(IConcertClientRef Client, const FString& SessionNameToRecover, bool bLiveDataOnly);
+void StartRecovery(TSharedRef<IConcertSyncClient> SyncClient, const FString& SessionNameToRecover, bool bLiveDataOnly);
 
 /** End the recovery flow. This can be called to abort the recovery process. */
 void EndRecovery();
