- Ensured the disaster recovery finite state machine persisted the session changes at least one frame after completing workspace sync to get a end frame event in between.

diff --git a/Engine/Plugins/Developer/Concert/DisasterRecoveryClient/Source/DisasterRecoveryClient/Private/DisasterRecoveryFSM.cpp b/Engine/Plugins/Developer/Concert/DisasterRecoveryClient/Source/DisasterRecoveryClient/Private/DisasterRecoveryFSM.cpp
--- a/Engine/Plugins/Developer/Concert/DisasterRecoveryClient/Source/DisasterRecoveryClient/Private/DisasterRecoveryFSM.cpp
+++ b/Engine/Plugins/Developer/Concert/DisasterRecoveryClient/Source/DisasterRecoveryClient/Private/DisasterRecoveryFSM.cpp
@@ -102,6 +102,7 @@ private:
 	TSharedPtr<uint8> FutureExecutionToken;
 	FText ErrorMessage;
 	FString ExitMessage;
+	int32 WaitedFrameCount = 0;
 
 	// States.
 	FState EnterState;                 // State to start from.
@@ -160,6 +161,7 @@ FDisasterRecoveryFSM::FDisasterRecoveryFSM(TSharedRef<IConcertSyncClient> InSync
 	RestoreAndJoinSessionState.OnExit  = [this]() { Client->OnSessionConnectionChanged().RemoveAll(this); };
 	SynchronizeState.OnEnter           = [this]() { IDisasterRecoveryClientModule::Get().GetClient()->GetWorkspace()->OnWorkspaceSynchronized().AddRaw(this, &FDisasterRecoveryFSM::OnWorkspaceSynchronized); };
 	SynchronizeState.OnExit            = [this]() { IDisasterRecoveryClientModule::Get().GetClient()->GetWorkspace()->OnWorkspaceSynchronized().RemoveAll(this); };
+	PersistChangesState.OnEnter        = [this]() { WaitedFrameCount = 0; };
 	PersistChangesState.OnTick         = [this]() { PersistRecoveredChanges(); };
 	ErrorState.OnEnter                 = [this]() { NextStatePending = nullptr; DisplayError(); };
 	ExitState.OnEnter                  = [=]()    { Terminate(); };
@@ -381,9 +383,13 @@ void FDisasterRecoveryFSM::RestoreAndJoinSession()
 
 void FDisasterRecoveryFSM::PersistRecoveredChanges()
 {
-	// Save live transactions to package, gather files changed in the Concert sandbox and apply the changes to the content directory.
-	SyncClient->PersistAllSessionChanges();
-	TransitTo(ExitState); // Disaster recovery process completed successfully.
+	// Don't execute in the same frame as OnWorkspaceSynchronized(). FConcertClientWorkspace::OnEndFrame() needs to run first to apply the transaction before persisting changes.
+	if (WaitedFrameCount++ >= 1)
+	{
+		// Save live transactions to package, gather files changed in the Concert sandbox and apply the changes to the content directory.
+		SyncClient->PersistAllSessionChanges();
+		TransitTo(ExitState); // Disaster recovery process completed successfully.
+	}
 }
 
 void FDisasterRecoveryFSM::DisplayError()
