- Cleared the concert server instance info on server shutdown.
  - Shutdown the disaster recovery service when a crash is created. This enable the next server instance to grab the file lock and restore.
  - Fixed archive rotation (delete oldest) that did not work when concurrent servers existed.
  - Improved disaster recovery error messages.
  - Fixed disaster recovery client not restoring a session that was crashed (server managed the crash), but for which the client process was still hanging around.
  - Prevent showing the recovery UI if -unattended is specified on command line.

diff --git a/Engine/Plugins/Developer/Concert/ConcertMain/Source/Concert/Private/ConcertServer.cpp b/Engine/Plugins/Developer/Concert/ConcertMain/Source/Concert/Private/ConcertServer.cpp
--- a/Engine/Plugins/Developer/Concert/ConcertMain/Source/Concert/Private/ConcertServer.cpp
+++ b/Engine/Plugins/Developer/Concert/ConcertMain/Source/Concert/Private/ConcertServer.cpp
@@ -26,6 +26,12 @@
 namespace ConcertServerUtil
 {
 
+const TCHAR* GetServerSystemMutexName()
+{
+	// A system wide mutex name used by this application instances that will unlikely be found in other applications.
+	return TEXT("Unreal_ConcertServer_67822dAB");
+}
+
 FString GetArchiveName(const FString& SessionName, const FConcertSessionSettings& Settings)
 {
 	if (Settings.ArchiveNameOverride.IsEmpty())
@@ -192,8 +198,8 @@ void FConcertServer::Startup()
 		ServerAdminEndpoint->RegisterRequestHandler<FConcertAdmin_GetSessionClientsRequest, FConcertAdmin_GetSessionClientsResponse>(this, &FConcertServer::HandleGetSessionClientsRequest);
 		ServerAdminEndpoint->RegisterRequestHandler<FConcertAdmin_GetSessionActivitiesRequest, FConcertAdmin_GetSessionActivitiesResponse>(this, &FConcertServer::HandleGetSessionActivitiesRequest);
 
-		// Concurrent server instances may fight to get ownership of the info file. Use a system wide mutex with an arbitrary name that will unlikely be found in other applications.
-		FSystemWideCriticalSection ScopedSystemWideMutex(TEXT("Unreal_ConcertServer_67822dAB"));
+		// Concurrent server instances may fight to get ownership of the info file.
+		FSystemWideCriticalSection ScopedSystemWideMutex(ConcertServerUtil::GetServerSystemMutexName());
 
 		// Load the file containing the instance info.
 		FConcertServerInstanceInfo InstanceInfo;
@@ -201,8 +207,8 @@ void FConcertServer::Startup()
 
 		// Check if the server can scan/load/rotate the existing session files in the configured directories. (The existing session files are not sharable, they can only managed by one instance)
 		int32 ActiveInstanceCount = ConcertServerUtil::RemoveInactiveServerInstances(InstanceInfo);
-		bool bLoadExistingSessions = ActiveInstanceCount == 0 || !ConcertServerUtil::SharesDirectoriesWithOtherInstances(*Paths, InstanceInfo);
-		if (bLoadExistingSessions)
+		bExclusiveAccessToExistingSessionsAtStartup = ActiveInstanceCount == 0 || !ConcertServerUtil::SharesDirectoriesWithOtherInstances(*Paths, InstanceInfo);
+		if (bExclusiveAccessToExistingSessionsAtStartup)
 		{
 			if (Settings->bCleanWorkingDir)
 			{
@@ -228,8 +234,7 @@ void FConcertServer::Startup()
 		ThisInstance.WorkingDirectory = Paths->GetWorkingDir();
 		ThisInstance.ArchiveDirectory = Paths->GetSavedDir();
 		ThisInstance.ProcessId = FPlatformProcess::GetCurrentProcessId();
-		ThisInstance.bEnclaved = !bLoadExistingSessions; // Enclaved means this instance did not scan and load existing sessions to avoid conflict with another running instance(s). It was restricted to work with its own sessions only.
-
+		ThisInstance.bEnclaved = !bExclusiveAccessToExistingSessionsAtStartup; // Enclaved means this instance did not scan and load existing sessions to avoid conflict with another running instance(s). It was restricted to work with its own sessions only.
 		ConcertServerUtil::SaveServerInstanceInfo(Role, InstanceInfo);
 	}
 }
@@ -286,6 +291,22 @@ void FConcertServer::Shutdown()
 		}
 		ArchivedSessions.Reset();
 	}
+
+	// Concurrent server instances may fight to get ownership of the info file.
+	FSystemWideCriticalSection ScopedSystemWideMutex(ConcertServerUtil::GetServerSystemMutexName());
+
+	// Load the file containing the instance info.
+	FConcertServerInstanceInfo InstanceInfo;
+	ConcertServerUtil::LoadServerInstanceInfo(Role, InstanceInfo);
+
+	// Remove this instance form the list of instances.
+	int32 ProcessId = FPlatformProcess::GetCurrentProcessId();
+	InstanceInfo.Instances.RemoveAll([ProcessId](const FConcertServerInstance& Instance)
+	{
+		return Instance.ProcessId == ProcessId;
+	});
+
+	ConcertServerUtil::SaveServerInstanceInfo(Role, InstanceInfo);
 }
 
 FGuid FConcertServer::GetLiveSessionIdByName(const FString& InName) const
@@ -416,9 +437,11 @@ void FConcertServer::RecoverSessions()
 	{
 		// Find any existing archived sessions
 		TArray<FConcertSessionInfo> ArchivedSessionInfos;
-		TArray<FDateTime> ArchivedSessionLastModifiedTimes;
-		EventSink->GetSessionsFromPath(*this, Paths->GetSavedDir(), ArchivedSessionInfos, &ArchivedSessionLastModifiedTimes);
-		check(ArchivedSessionInfos.Num() == ArchivedSessionLastModifiedTimes.Num());
+		TArray<FDateTime> ArchivedSessionCreationTimes;
+
+		// In theory, archives are immutable, but the server will end up touching the files and change the 'modification time'. Ensure to look at 'creation time'.
+		EventSink->GetSessionsFromPath(*this, Paths->GetSavedDir(), ArchivedSessionInfos, &ArchivedSessionCreationTimes);
+		check(ArchivedSessionInfos.Num() == ArchivedSessionCreationTimes.Num());
 
 		// Trim the oldest archived sessions
 		if (Settings->NumSessionsToKeep > 0 && ArchivedSessionInfos.Num() > Settings->NumSessionsToKeep)
@@ -429,7 +452,7 @@ void FConcertServer::RecoverSessions()
 			TArray<FSavedSessionInfo> SortedSessions;
 			for (int32 LiveSessionInfoIndex = 0; LiveSessionInfoIndex < ArchivedSessionInfos.Num(); ++LiveSessionInfoIndex)
 			{
-				SortedSessions.Add(MakeTuple(LiveSessionInfoIndex, ArchivedSessionLastModifiedTimes[LiveSessionInfoIndex]));
+				SortedSessions.Add(MakeTuple(LiveSessionInfoIndex, ArchivedSessionCreationTimes[LiveSessionInfoIndex]));
 			}
 			SortedSessions.Sort([](const FSavedSessionInfo& InOne, const FSavedSessionInfo& InTwo)
 			{
@@ -455,7 +478,7 @@ void FConcertServer::RecoverSessions()
 
 			// Update the list of sessions to restore
 			ArchivedSessionInfos = MoveTemp(ArchivedSessionsToKeep);
-			ArchivedSessionLastModifiedTimes.Reset();
+			ArchivedSessionCreationTimes.Reset();
 		}
 
 		// Create any existing archived sessions
@@ -866,6 +889,7 @@ TFuture<FConcertAdmin_GetAllSessionsResponse> FConcertServer::HandleGetAllSessio
 
 	FConcertAdmin_GetAllSessionsResponse ResponseData;
 	ResponseData.LiveSessions = GetSessionsInfo();
+	ResponseData.bIncludesPreviousInstanceSessions = bExclusiveAccessToExistingSessionsAtStartup;
 	for (const auto& ArchivedSessionPair : ArchivedSessions)
 	{
 		ResponseData.ArchivedSessions.Add(ArchivedSessionPair.Value);
diff --git a/Engine/Plugins/Developer/Concert/ConcertMain/Source/Concert/Private/ConcertServer.h b/Engine/Plugins/Developer/Concert/ConcertMain/Source/Concert/Private/ConcertServer.h
--- a/Engine/Plugins/Developer/Concert/ConcertMain/Source/Concert/Private/ConcertServer.h
+++ b/Engine/Plugins/Developer/Concert/ConcertMain/Source/Concert/Private/ConcertServer.h
@@ -213,4 +213,7 @@ private:
 
 	/** Server settings object we were configured with */
 	TStrongObjectPtr<const UConcertServerConfig> Settings;
+
+	/** Whether the server could scan/rotate/load the existing sessions (not-sharable) at startup because no other concurrent server was running. */
+	bool bExclusiveAccessToExistingSessionsAtStartup = true;
 };
diff --git a/Engine/Plugins/Developer/Concert/ConcertMain/Source/Concert/Public/ConcertMessages.h b/Engine/Plugins/Developer/Concert/ConcertMain/Source/Concert/Public/ConcertMessages.h
--- a/Engine/Plugins/Developer/Concert/ConcertMain/Source/Concert/Public/ConcertMessages.h
+++ b/Engine/Plugins/Developer/Concert/ConcertMain/Source/Concert/Public/ConcertMessages.h
@@ -108,6 +108,13 @@ struct FConcertAdmin_GetAllSessionsResponse : public FConcertResponseData
 
 	UPROPERTY(VisibleAnywhere, Category="Concert Message")
 	TArray<FConcertSessionInfo> ArchivedSessions;
+
+	/**
+	 * True when the queried server sees the sessions created/archived during its life time and the sessions from previous (now dead) instances.
+	 * Concurrent servers instances will not see the same sessions because sessions files are not-sharable.
+	 */
+	UPROPERTY(VisibleAnywhere, Category="Concert Message")
+	bool bIncludesPreviousInstanceSessions = true;
 };
 
 USTRUCT()
diff --git a/Engine/Plugins/Developer/Concert/ConcertMain/Source/Concert/Public/IConcertServerEventSink.h b/Engine/Plugins/Developer/Concert/ConcertMain/Source/Concert/Public/IConcertServerEventSink.h
--- a/Engine/Plugins/Developer/Concert/ConcertMain/Source/Concert/Public/IConcertServerEventSink.h
+++ b/Engine/Plugins/Developer/Concert/ConcertMain/Source/Concert/Public/IConcertServerEventSink.h
@@ -18,7 +18,7 @@ public:
 	/**
 	 * Called to enumerate all the sessions under the given root path and retrieve their session info.
 	 */
-	virtual void GetSessionsFromPath(const IConcertServer& InServer, const FString& InPath, TArray<FConcertSessionInfo>& OutSessionInfos, TArray<FDateTime>* OutSessionLastModifiedTimes = nullptr) = 0;
+	virtual void GetSessionsFromPath(const IConcertServer& InServer, const FString& InPath, TArray<FConcertSessionInfo>& OutSessionInfos, TArray<FDateTime>* OutSessionCreationTimes = nullptr) = 0;
 
 	/**
 	 * Called after the session has been created (and before Startup has been called on it).
diff --git a/Engine/Plugins/Developer/Concert/ConcertSync/ConcertSyncServer/Source/ConcertSyncServer/Private/ConcertSyncServer.cpp b/Engine/Plugins/Developer/Concert/ConcertSync/ConcertSyncServer/Source/ConcertSyncServer/Private/ConcertSyncServer.cpp
--- a/Engine/Plugins/Developer/Concert/ConcertSync/ConcertSyncServer/Source/ConcertSyncServer/Private/ConcertSyncServer.cpp
+++ b/Engine/Plugins/Developer/Concert/ConcertSync/ConcertSyncServer/Source/ConcertSyncServer/Private/ConcertSyncServer.cpp
@@ -24,26 +24,16 @@
 namespace ConcertSyncServerUtils
 {
 
-const TCHAR SessionDatabaseFilename[] = TEXT("Session.db");
 const TCHAR SessionInfoFilename[] = TEXT("SessionInfo.json");
 
-/** Attempt to get the last modification time for a session */
-FDateTime GetSessionModificationTime(const FString& InSessionRoot)
+/** Deduce the creation time of a session from the creation time of its containing folder. */
+FDateTime GetSessionCreationTime(const FString& InSessionRoot)
 {
 	IFileManager& FileManager = IFileManager::Get();
 
-	const FString SessionDatabasePathname = InSessionRoot / SessionDatabaseFilename;
-	const FString SessionInfoPathname = InSessionRoot / SessionInfoFilename;
-
-	// Try the database first, as that will be written as the session is in progress
-	FDateTime SessionModificationTime = FileManager.GetTimeStamp(*SessionDatabasePathname);
-	if (SessionModificationTime == FDateTime::MinValue())
-	{
-		// Failing that try the info file, as that will be written whenever the session is initialized
-		SessionModificationTime = FileManager.GetTimeStamp(*SessionInfoPathname);
-	}
-
-	return SessionModificationTime;
+	FFileStatData FileStatData = FileManager.GetStatData(*InSessionRoot);
+	check(FileStatData.bIsValid && FileStatData.CreationTime != FDateTime::MinValue());
+	return FileStatData.CreationTime;
 }
 
 /** Write the session info file to a session root directory */
@@ -242,9 +232,9 @@ IConcertServerRef FConcertSyncServer::GetConcertServer() const
 	return ConcertServer;
 }
 
-void FConcertSyncServer::GetSessionsFromPath(const IConcertServer& InServer, const FString& InPath, TArray<FConcertSessionInfo>& OutSessionInfos, TArray<FDateTime>* OutSessionLastModifiedTimes)
+void FConcertSyncServer::GetSessionsFromPath(const IConcertServer& InServer, const FString& InPath, TArray<FConcertSessionInfo>& OutSessionInfos, TArray<FDateTime>* OutSessionCreationTimes)
 {
-	IFileManager::Get().IterateDirectory(*InPath, [&OutSessionInfos, &OutSessionLastModifiedTimes](const TCHAR* FilenameOrDirectory, bool bIsDirectory) -> bool
+	IFileManager::Get().IterateDirectory(*InPath, [&OutSessionInfos, &OutSessionCreationTimes](const TCHAR* FilenameOrDirectory, bool bIsDirectory) -> bool
 	{
 		if (bIsDirectory)
 		{
@@ -260,9 +250,9 @@ void FConcertSyncServer::GetSessionsFromPath(const IConcertServer& InServer, con
 				if (FGuid::Parse(SessionFolderName, SessionFolderGuid) && SessionFolderGuid == SessionInfo.SessionId)
 				{
 					OutSessionInfos.Add(MoveTemp(SessionInfo));
-					if (OutSessionLastModifiedTimes)
+					if (OutSessionCreationTimes)
 					{
-						OutSessionLastModifiedTimes->Add(ConcertSyncServerUtils::GetSessionModificationTime(SessionRoot));
+						OutSessionCreationTimes->Add(ConcertSyncServerUtils::GetSessionCreationTime(SessionRoot));
 					}
 				}
 			}
diff --git a/Engine/Plugins/Developer/Concert/ConcertSync/ConcertSyncServer/Source/ConcertSyncServer/Private/ConcertSyncServer.h b/Engine/Plugins/Developer/Concert/ConcertSync/ConcertSyncServer/Source/ConcertSyncServer/Private/ConcertSyncServer.h
--- a/Engine/Plugins/Developer/Concert/ConcertSync/ConcertSyncServer/Source/ConcertSyncServer/Private/ConcertSyncServer.h
+++ b/Engine/Plugins/Developer/Concert/ConcertSync/ConcertSyncServer/Source/ConcertSyncServer/Private/ConcertSyncServer.h
@@ -29,7 +29,7 @@ public:
 	virtual IConcertServerRef GetConcertServer() const override;
 
 	//~ IConcertServerEventSink interface
-	virtual void GetSessionsFromPath(const IConcertServer& InServer, const FString& InPath, TArray<FConcertSessionInfo>& OutSessionInfos, TArray<FDateTime>* OutSessionLastModifiedTimes = nullptr) override;
+	virtual void GetSessionsFromPath(const IConcertServer& InServer, const FString& InPath, TArray<FConcertSessionInfo>& OutSessionInfos, TArray<FDateTime>* OutSessionCreationTimes = nullptr) override;
 	virtual void OnLiveSessionCreated(const IConcertServer& InServer, TSharedRef<IConcertServerSession> InLiveSession) override;
 	virtual void OnLiveSessionDestroyed(const IConcertServer& InServer, TSharedRef<IConcertServerSession> InLiveSession) override;
 	virtual void OnArchivedSessionCreated(const IConcertServer& InServer, const FString& InArchivedSessionRoot, const FConcertSessionInfo& InArchivedSessionInfo) override;
diff --git a/Engine/Plugins/Developer/Concert/DisasterRecoveryClient/Source/DisasterRecoveryClient/Private/DisasterRecoveryClientModule.cpp b/Engine/Plugins/Developer/Concert/DisasterRecoveryClient/Source/DisasterRecoveryClient/Private/DisasterRecoveryClientModule.cpp
--- a/Engine/Plugins/Developer/Concert/DisasterRecoveryClient/Source/DisasterRecoveryClient/Private/DisasterRecoveryClientModule.cpp
+++ b/Engine/Plugins/Developer/Concert/DisasterRecoveryClient/Source/DisasterRecoveryClient/Private/DisasterRecoveryClientModule.cpp
@@ -201,7 +201,7 @@ private:
 		return FPaths::ProjectSavedDir() / Role / TEXT("SessionInfo.json");
 	}
 
-	bool LoadDisasterRecoverySessionInfo(FDisasterRecoverySessionInfo& OutSessionInfo)
+	bool LoadDisasterRecoverySessionInfo(FDisasterRecoverySessionInfo& OutSessionInfo) const
 	{
 		if (TUniquePtr<FArchive> FileReader = TUniquePtr<FArchive>(IFileManager::Get().CreateFileReader(*GetDisasterRecoverySessionInfoFilename())))
 		{
@@ -215,7 +215,7 @@ private:
 		return false;
 	}
 
-	bool SaveDisasterRecoverySessionInfo(const FDisasterRecoverySessionInfo& InSessionInfo)
+	bool SaveDisasterRecoverySessionInfo(const FDisasterRecoverySessionInfo& InSessionInfo) const
 	{
 		if (TUniquePtr<FArchive> FileWriter = TUniquePtr<FArchive>(IFileManager::Get().CreateFileWriter(*GetDisasterRecoverySessionInfoFilename())))
 		{
@@ -233,7 +233,7 @@ private:
 	{
 		auto GetDisasterRecoveryServicePathForBuildConfiguration = [](const EBuildConfiguration InBuildConfiguration) -> FString
 		{
-			FString ServicePath = FPlatformProcess::GenerateApplicationPath(TEXT("UnrealDisasterRecoveryService"), InBuildConfiguration);
+			FString ServicePath = FPlatformProcess::GenerateApplicationPath(DisasterRecoveryUtil::GetDisasterRecoveryServiceExeName(), InBuildConfiguration);
 			return FPaths::FileExists(ServicePath) ? ServicePath : FString();
 		};
 
@@ -264,21 +264,15 @@ private:
 		LoadDisasterRecoverySessionInfo(RecoverySessionInfo);
 
 		// Is this a new session created by recovering from another one?
-		if (RestoringSession.IsSet())
+		RecoverySessionInfo.Sessions.RemoveAll([](const FDisasterRecoverySession& Session)
 		{
-			// Remove the recovered session from the list.
-			RecoverySessionInfo.Sessions.RemoveAll([this](const FDisasterRecoverySession& Session)
-			{
-				return Session.bAutoRestoreLastSession && RestoringSession->ProcessId == Session.ProcessId && RestoringSession->LastSessionName == Session.LastSessionName;
-			});
-
-			RestoringSession.Reset();
-		}
+			return Session.bAutoRestoreLastSession && Session.HostProcessId == FPlatformProcess::GetCurrentProcessId();
+		});
 
 		// Create a new session.
 		FDisasterRecoverySession& Session = RecoverySessionInfo.Sessions.AddDefaulted_GetRef();
 		Session.LastSessionName = InSession->GetSessionInfo().SessionName;
-		Session.ProcessId = FPlatformProcess::GetCurrentProcessId();
+		Session.HostProcessId = FPlatformProcess::GetCurrentProcessId();
 
 		// Normally, bAutoRestoreLastSession is set true here and overwritten to false when the app exits normally, but when running under the debugger, auto-restore is disabled as
 		// programmers kill applications (stop the debugger) and this should not count as a crash (unless you want to simulate crash this way - see below).
@@ -289,43 +283,65 @@ private:
 		SaveDisasterRecoverySessionInfo(RecoverySessionInfo);
 	}
 
-	TOptional<FDisasterRecoverySession> ProcessSessionInfo()
+	TPair<bool, const FConcertSessionInfo*> PinSessionToRestore(const TArray<FConcertSessionInfo>& ArchivedSessions)
 	{
-		// Ensure we get exclusive access to the file.
+		// Ensure we get exclusive access to the recovery session info file.
 		FSystemWideCriticalSection SystemWideMutex(GetSystemMutexName());
 
-		TOptional<FDisasterRecoverySession> SessionToRestore;
-
 		// Load the session info file (if it exist)
 		FDisasterRecoverySessionInfo RecoverySessionInfo;
 		LoadDisasterRecoverySessionInfo(RecoverySessionInfo);
 
+		// Checks if two running process ID are instances of same executable.
+		auto IsSameExecutable = [](int32 LhsProcessId, int32 RhsProcessId)
+		{
+			return FPaths::GetPathLeaf(FPlatformProcess::GetApplicationName(LhsProcessId)) == FPaths::GetPathLeaf(FPlatformProcess::GetApplicationName(RhsProcessId));
+		};
+
+		// Returns true if the process hosting the session crashed. Note that CrashReporterClientEditor will shut down the server, archive the session and relaunch a new editor (and a new recovery client) before the
+		// previous editor finished crashing. The process hosting a session may still be alive, but not its server. In such case, the session will be unlocked and available for restoration unless another server instances already exist.
+		auto IsHostProcessCrashed = [&IsSameExecutable](const FDisasterRecoverySession& Session)
+		{
+			return Session.bAutoRestoreLastSession && (!FPlatformProcess::IsApplicationRunning(Session.HostProcessId) || !IsSameExecutable(Session.HostProcessId, FPlatformProcess::GetCurrentProcessId()));
+		};
+
+		// Returns true if the session was listed as archived by the server. (Might not be visible to that server if other servers are running and use the same working/saved directories, the files are not sharable)
+		auto IsSessionArchived = [](const FDisasterRecoverySession& Session, const TArray<FConcertSessionInfo>& ArchivedSessions)
+		{
+			return ArchivedSessions.FindByPredicate([&Session](const FConcertSessionInfo& MatchCandidate) { return MatchCandidate.SessionName == Session.LastSessionName; });
+		};
+
+		// Check in the client file for session candidate to restoration.
+		TPair<bool, const FConcertSessionInfo*> RestoreCandidate = MakeTuple(false, static_cast<const FConcertSessionInfo*>(nullptr));
 		for (FDisasterRecoverySession& Session : RecoverySessionInfo.Sessions)
 		{
-			if (Session.bAutoRestoreLastSession)
+			const FConcertSessionInfo* ArchivedSession = IsSessionArchived(Session, ArchivedSessions);
+			if (ArchivedSession || IsHostProcessCrashed(Session))
 			{
-				if (!FPlatformProcess::IsApplicationRunning(Session.ProcessId) || // If the session owner is not running anymore.
-					FPaths::GetPathLeaf(FPlatformProcess::GetApplicationName(Session.ProcessId)) != FPaths::GetPathLeaf(FPlatformProcess::GetApplicationName(FPlatformProcess::GetCurrentProcessId()))) // This owner PID is live, but was reused by the OS for another app.
+				// Already pinned a session to restore?
+				if (RestoreCandidate.Key && RestoreCandidate.Value)
 				{
-					if (!SessionToRestore.IsSet())
-					{
-						// Recover this session. Take ownership of this session. (Will be cleared on DisasterRecoverySessionCreated() if recovery succeed)
-						Session.ProcessId = FPlatformProcess::GetCurrentProcessId();
+					// User ran multiple instances of the Editor on the same project and more than one instance crashed. Clear this one and restore the first one found only.
+					// The situation is unlikely to happen, but a future task would be to return a list of crashed sessions, pass it to the DisasterRecoveryFSM and prompt the user to select which one to recover.
+					Session.bAutoRestoreLastSession = false;
+				}
+				else
+				{
+					// A crashed session exist.
+					RestoreCandidate.Key = true;
 
-						// Pick up the first crashed session. (This is arbitrary if they are more)
-						SessionToRestore.Emplace(Session);
-					}
-					else
+					// Can this session be restored (i.e. the server knowns it)?
+					if (ArchivedSession)
 					{
-						// User ran multiple instances of the Editor on the same project and more than one instance crashed. Just restore the first one found above for the moment.
-						// While the situation above is unlikely to happen, a future task would be to return a list of crashed session, pass it to the DisasterRecoveryFSM and prompt the user to which one to recover.
-						Session.bAutoRestoreLastSession = false;
+						// Pin/take ownership of this session. (Will be cleared on DisasterRecoverySessionCreated() if recovery succeed)
+						RestoreCandidate.Value = ArchivedSession;
+						Session.HostProcessId = FPlatformProcess::GetCurrentProcessId();
 					}
 				}
 			}
 		}
 
-		// Remove completed sessions.
+		// Remove completed sessions (in case a session was flagged bAutoRestoreLastSession = false above)
 		RecoverySessionInfo.Sessions.RemoveAll([](const FDisasterRecoverySession& Session)
 		{
 			return !Session.bAutoRestoreLastSession; // Remove if the 'restore' flag is false.
@@ -334,7 +350,30 @@ private:
 		// Save the recovery session info file.
 		SaveDisasterRecoverySessionInfo(RecoverySessionInfo);
 
-		return SessionToRestore;
+		return RestoreCandidate;
+	}
+
+	bool HasCandidateSessionToRestore() const
+	{
+		// Ensure we get exclusive access to the recovery session info file.
+		FSystemWideCriticalSection SystemWideMutex(GetSystemMutexName());
+
+		// Load the session info file (if it exist)
+		FDisasterRecoverySessionInfo RecoverySessionInfo;
+		LoadDisasterRecoverySessionInfo(RecoverySessionInfo);
+
+		// If at least one session exist, it may be candidate to restore.
+		return RecoverySessionInfo.Sessions.Num() != 0;
+	}
+
+	void ClearSessionInfoFile()
+	{
+		// Ensure we get exclusive access to the recovery session info file.
+		FSystemWideCriticalSection SystemWideMutex(GetSystemMutexName());
+
+		// Load the session info file (if it exist)
+		FDisasterRecoverySessionInfo RecoverySessionInfo;
+		SaveDisasterRecoverySessionInfo(RecoverySessionInfo);
 	}
 
 	bool SpawnDisasterRecoveryServer(const FString& ServerName) // On Linux and Mac. On Windows, it is embedded in CrashReporterClient.
@@ -391,14 +430,21 @@ private:
 			return false; // Failed to spawn the service.
 		}
 
-		// Clean sessions and find if a session should be restored
-		RestoringSession = ProcessSessionInfo();
+		// It is not allowed to prompt the user for recovery in unattended mode. Forget everything about previous session(s) for this project and
+		// don't to keep crashed session(s) on 'hold' to be restored the next time -unattended is not specified, the levels may have greatly changed.
+		if (FApp::IsUnattended())
+		{
+			ClearSessionInfoFile();
+		}
+
+		// Skip restoration if they are no candidates to restore.
+		bool bCandidateToRestore = HasCandidateSessionToRestore();
 
 		// Create and populate the client config object
 		UConcertClientConfig* ClientConfig = NewObject<UConcertClientConfig>();
 		ClientConfig->bIsHeadless = true;
 		ClientConfig->bInstallEditorToolbarButton = false;
-		ClientConfig->bAutoConnect = !RestoringSession.IsSet(); // If recovering from a crash, don't auto connect -> Present UI to let user decide what to recover first.
+		ClientConfig->bAutoConnect = !bCandidateToRestore; // If there is a possible candidate to restore, don't auto connect yet, run the finite state machine to know more.
 		ClientConfig->DefaultServerURL = DisasterRecoveryServerName;
 		ClientConfig->DefaultSessionName = DisasterRecoverySessionName;
 		ClientConfig->DefaultSaveSessionAs = DisasterRecoverySessionName;
@@ -410,11 +456,11 @@ private:
 		DisasterRecoveryClient->GetConcertClient()->OnSessionStartup().AddRaw(this, &FDisasterRecoveryClientModule::DisasterRecoverySessionCreated);
 		DisasterRecoveryClient->Startup(ClientConfig, EConcertSyncSessionFlags::Default_DisasterRecoverySession);
 
-		// Set all events captured by the disaster recovery service as 'replayable' unless another concert client (assumed Multi-User) has created an incompatible session.
+		// Set all events captured by the disaster recovery service as 'restorable' unless another concert client (assumed Multi-User) has created an incompatible session.
 		SetIgnoreOnRestoreState(!IsCompatibleWithOtherConcertSessions(/*SyncClientStartingSession*/nullptr, /*SyncClientShuttingDownSession*/nullptr));
 
-		// If something needs to be recovered from a crash.
-		if (RestoringSession)
+		// If something might be recovered from a crash.
+		if (bCandidateToRestore)
 		{
 			if (GUnrealEd)
 			{
@@ -422,7 +468,13 @@ private:
 				GUnrealEd->GetPackageAutoSaver().DisableRestorePromptAndDeclinePackageRecovery();
 			}
 
-			DisasterRecoveryUtil::StartRecovery(DisasterRecoveryClient.ToSharedRef(), RestoringSession->LastSessionName, /*bLiveDataOnly*/ false);
+			// The DisasterRecoveryFSM will call this function to select and take ownership of the session to recover.
+			auto PinSessionFn = [this](const TArray<FConcertSessionInfo>& ArchivedSessions)
+			{
+				return PinSessionToRestore(ArchivedSessions);
+			};
+
+			DisasterRecoveryUtil::StartRecovery(DisasterRecoveryClient.ToSharedRef(), PinSessionFn, /*bLiveDataOnly*/ false);
 		}
 
 		return true;
@@ -444,7 +496,7 @@ private:
 			int32 ProcessId = FPlatformProcess::GetCurrentProcessId();
 			RecoverySessionInfo.Sessions.RemoveAll([ProcessId](const FDisasterRecoverySession& Session)
 			{
-				return Session.ProcessId == ProcessId;
+				return Session.HostProcessId == ProcessId;
 			});
 
 			// Write the file to disk.
@@ -521,9 +573,6 @@ private:
 
 	/** Handle to the active disaster recovery service app, if any */
 	FProcHandle DisasterRecoveryServiceHandle;
-
-	/** Keep the session being restored, if any. */
-	TOptional<FDisasterRecoverySession> RestoringSession;
 };
 
 #undef LOCTEXT_NAMESPACE
diff --git a/Engine/Plugins/Developer/Concert/DisasterRecoveryClient/Source/DisasterRecoveryClient/Private/DisasterRecoveryFSM.cpp b/Engine/Plugins/Developer/Concert/DisasterRecoveryClient/Source/DisasterRecoveryClient/Private/DisasterRecoveryFSM.cpp
--- a/Engine/Plugins/Developer/Concert/DisasterRecoveryClient/Source/DisasterRecoveryClient/Private/DisasterRecoveryFSM.cpp
+++ b/Engine/Plugins/Developer/Concert/DisasterRecoveryClient/Source/DisasterRecoveryClient/Private/DisasterRecoveryFSM.cpp
@@ -15,6 +15,7 @@
 #include "Widgets/SWindow.h"
 #include "SConcertSessionRecovery.h"
 #include "ConcertActivityStream.h"
+#include "GenericPlatform/GenericPlatformCrashContext.h"
 
 #define LOCTEXT_NAMESPACE "DisasterRecoveryFSM"
 
@@ -29,7 +30,7 @@ public:
 	static constexpr int64 MaxActivityPerRequest = 1024;
 
 public:
-	FDisasterRecoveryFSM(TSharedRef<IConcertSyncClient> InSyncClient, const FString& InSessionNameToRecover, bool bInLiveDataOnly);
+	FDisasterRecoveryFSM(TSharedRef<IConcertSyncClient> InSyncClient, const DisasterRecoveryUtil::PinSessionToRestoreFunc& PinSessionToRestoreFn, bool bInLiveDataOnly);
 	~FDisasterRecoveryFSM();
 
 	bool IsDone() const
@@ -131,18 +132,20 @@ private:
 	FState ExitState;                  // State to exit the FSM.
 	FState* CurrentState = nullptr;    // The current state.
 	FState* NextStatePending = nullptr;// State to transit to at the next Tick().
+
+	DisasterRecoveryUtil::PinSessionToRestoreFunc PinSessionToRestoreFn;
 };
 
 static TSharedPtr<FDisasterRecoveryFSM> RecoveryFSM;
 
 
-FDisasterRecoveryFSM::FDisasterRecoveryFSM(TSharedRef<IConcertSyncClient> InSyncClient, const FString& SessionNameToRecover, bool bInLiveDataOnly)
+FDisasterRecoveryFSM::FDisasterRecoveryFSM(TSharedRef<IConcertSyncClient> InSyncClient, const DisasterRecoveryUtil::PinSessionToRestoreFunc& InPinSessionToRestoreFn, bool bInLiveDataOnly)
 	: SyncClient(InSyncClient)
 	, Client(InSyncClient->GetConcertClient())
-	, RecoverySessionName(SessionNameToRecover)
 	, bLiveDataOnly(bInLiveDataOnly)
 	, ExitMessage(TEXT("Disaster recovery process completed successfully."))
 	, CurrentState(&EnterState)
+	, PinSessionToRestoreFn(InPinSessionToRestoreFn)
 {
 	FutureExecutionToken = MakeShared<uint8>(0); // Token use to prevent execution of async continuation if this object instance is destructed before then continuation is executed.
 
@@ -203,8 +206,8 @@ void FDisasterRecoveryFSM::OnSessionConnectionChanged(IConcertClientSession& Ses
 	}
 	else if (Status == EConcertConnectionStatus::Disconnected)
 	{
-		ErrorMessage = LOCTEXT("ConnectionError", "Disaster recover client failed to connect to the recovery session.");
-		TransitTo(ErrorState); // Something went wrong.
+		ErrorMessage = LOCTEXT("ConnectionError", "Failed to connect to the recovery session. Recovery service will be disabled for this session.");
+		TransitTo(ErrorState);
 	}
 }
 
@@ -241,21 +244,32 @@ void FDisasterRecoveryFSM::LookupRecoverySession()
 
 		if (Response.ResponseCode == EConcertResponseCode::Success)
 		{
-			// The disaster recovery service archives the sessions on exit (or on reboot if it crashed/was killed).
-			if (const FConcertSessionInfo* ArchMatch = Response.ArchivedSessions.FindByPredicate([this](const FConcertSessionInfo& MatchCandidate) { return MatchCandidate.SessionName == RecoverySessionName; }))
+			// Compare the archived sessions received from the server with the candidate sessions known by the client and select/pin a session to recover.
+			TPair<bool, const FConcertSessionInfo*> SessionToRecover = PinSessionToRestoreFn(Response.ArchivedSessions);
+			if (SessionToRecover.Value) // Took ownership of a session to restore?
 			{
-				RecoverySessionId = ArchMatch->SessionId;
+				RecoverySessionName = SessionToRecover.Value->SessionName;
+				RecoverySessionId = SessionToRecover.Value->SessionId;
 				RequestTransitTo(FetchActivitiesState);
 			}
+			else if (SessionToRecover.Key && !Response.bIncludesPreviousInstanceSessions) // The client has a candidate to restore, but the server did not report it.
+			{
+				ErrorMessage = FText::Format(LOCTEXT("RecoveryFileLocked", "Failed to recover previous session. The session is likely locked by another instance of '{0}'. Exit all Editors then relaunch to recover. Recovery service will be disabled for this session."), FText::AsCultureInvariant(DisasterRecoveryUtil::GetDisasterRecoveryServiceExeName()));
+				RequestTransitTo(ErrorState);
+			}
+			else if (SessionToRecover.Key && Response.bIncludesPreviousInstanceSessions)
+			{
+				UE_LOG(LogConcert, Error, TEXT("Failed to recover previous session. The session was likely moved or deleted. A new session will be created."));
+				RequestTransitTo(CreateAndJoinSessionState); // No candidate to recover. Create a new session.
+			}
 			else
 			{
-				UE_LOG(LogConcert, Error, TEXT("Disaster recovery service could not find the recovery session '%s'. A new session will be created."), *RecoverySessionName);
-				RequestTransitTo(CreateAndJoinSessionState); // We still need to have a recovery session at the end.
+				RequestTransitTo(CreateAndJoinSessionState); // No candidate to recover. Create a new session.
 			}
 		}
 		else
 		{
-			ErrorMessage = LOCTEXT("SessionQueryFailed", "Failed to retrieve the session list from the recovery service.");
+			ErrorMessage = LOCTEXT("SessionQueryFailed", "Failed to retrieve available sessions. Recovery service will be disabled for this session.");
 			RequestTransitTo(ErrorState);
 		}
 	});
@@ -282,7 +296,7 @@ void FDisasterRecoveryFSM::FetchActivities()
 		}
 		else // Nothing to recover
 		{
-			UE_LOG(LogConcert, Warning, TEXT("Disaster recovery service could not find any activity to recover"));
+			UE_LOG(LogConcert, Warning, TEXT("Disaster recovery service could not find any activities to recover."));
 			RequestTransitTo(CreateAndJoinSessionState);
 		}
 	}
@@ -351,7 +365,7 @@ void FDisasterRecoveryFSM::CreateAndJoinSession()
 
 		if (Response != EConcertResponseCode::Success)
 		{
-			ErrorMessage = FText::Format(LOCTEXT("FailedToCreate", "Failed to create recovery session '{0}'"), FText::AsCultureInvariant(Client->GetConfiguration()->DefaultSessionName));
+			ErrorMessage = FText::Format(LOCTEXT("FailedToCreate", "Failed to create recovery session '{0}'. Recovery service will be disabled for this session."), FText::AsCultureInvariant(Client->GetConfiguration()->DefaultSessionName));
 			RequestTransitTo(ErrorState);
 		}
 		// else -> On success callback OnSessionConnectionChanged() will transit to synchronize state.
@@ -437,11 +451,11 @@ bool FDisasterRecoveryFSM::Tick(float)
 	return true;
 }
 
-void DisasterRecoveryUtil::StartRecovery(TSharedRef<IConcertSyncClient> SyncClient, const FString& SessionNameToRecover, bool bLiveDataOnly)
+void DisasterRecoveryUtil::StartRecovery(TSharedRef<IConcertSyncClient> SyncClient, const DisasterRecoveryUtil::PinSessionToRestoreFunc& PinSessionToRestoreFn, bool bLiveDataOnly)
 {
 	if (!RecoveryFSM)
 	{
-		RecoveryFSM = MakeShared<FDisasterRecoveryFSM>(SyncClient, SessionNameToRecover, bLiveDataOnly);
+		RecoveryFSM = MakeShared<FDisasterRecoveryFSM>(SyncClient, PinSessionToRestoreFn, bLiveDataOnly);
 	}
 }
 
@@ -457,4 +471,16 @@ bool DisasterRecoveryUtil::EndRecovery()
 	return bDone;
 }
 
+FString DisasterRecoveryUtil::GetDisasterRecoveryServiceExeName()
+{
+	if (FGenericCrashContext::IsOutOfProcessCrashReporter())
+	{
+		return TEXT("CrashReporterClientEditor");
+	}
+	else
+	{
+		return TEXT("UnrealDisasterRecoveryService");
+	}
+}
+
 #undef LOCTEXT_NAMESPACE
diff --git a/Engine/Plugins/Developer/Concert/DisasterRecoveryClient/Source/DisasterRecoveryClient/Private/DisasterRecoveryFSM.h b/Engine/Plugins/Developer/Concert/DisasterRecoveryClient/Source/DisasterRecoveryClient/Private/DisasterRecoveryFSM.h
--- a/Engine/Plugins/Developer/Concert/DisasterRecoveryClient/Source/DisasterRecoveryClient/Private/DisasterRecoveryFSM.h
+++ b/Engine/Plugins/Developer/Concert/DisasterRecoveryClient/Source/DisasterRecoveryClient/Private/DisasterRecoveryFSM.h
@@ -3,20 +3,26 @@
 #pragma once
 
 #include "CoreMinimal.h"
+#include "ConcertMessageData.h"
 #include "IConcertModule.h"
+#include "Templates/Function.h"
 
 class IConcertSyncClient;
 
 namespace DisasterRecoveryUtil
 {
+/**
+ * Function invoked by by the recovery system to pin/lock a given session to be recovered by this instance.
+ */
+using PinSessionToRestoreFunc = TFunction<TPair<bool, const FConcertSessionInfo*>(const TArray<FConcertSessionInfo>& /*ArchivedSessions*/)>;
 
 /**
  * Start the recovery flow. This should be called when the previous instance of the editor exited unexpectedly.
  * @param SyncClient The client configured with the recovery server URL, the recovery session name and recovery archive name.
- * @param SessionNameToRecover The name of the session to recover.
+ * @param PinSessionToRestoreFn A function used by the process to take ownership for recovering a session.
  * @param bLiveDataOnly Filter the recovery data to only recover live transaction data. (Transaction that were not saved to disk yet)
  */
-void StartRecovery(TSharedRef<IConcertSyncClient> SyncClient, const FString& SessionNameToRecover, bool bLiveDataOnly);
+void StartRecovery(TSharedRef<IConcertSyncClient> SyncClient, const PinSessionToRestoreFunc& PinSessionToRestoreFn, bool bLiveDataOnly);
 
 /**
  * End the recovery flow. This can be called to abort the recovery process.
@@ -24,4 +30,8 @@ void StartRecovery(TSharedRef<IConcertSyncClient> SyncClient, const FString& Ses
  */
 bool EndRecovery();
 
+/**
+ * Return the name of the executable hosting disaster recovery service, like 'UnrealDisasterRecoveryService' without the extension.
+ */
+FString GetDisasterRecoveryServiceExeName();
 }
diff --git a/Engine/Plugins/Developer/Concert/DisasterRecoveryClient/Source/DisasterRecoveryClient/Private/DisasterRecoverySessionInfo.h b/Engine/Plugins/Developer/Concert/DisasterRecoveryClient/Source/DisasterRecoveryClient/Private/DisasterRecoverySessionInfo.h
--- a/Engine/Plugins/Developer/Concert/DisasterRecoveryClient/Source/DisasterRecoveryClient/Private/DisasterRecoverySessionInfo.h
+++ b/Engine/Plugins/Developer/Concert/DisasterRecoveryClient/Source/DisasterRecoveryClient/Private/DisasterRecoverySessionInfo.h
@@ -13,19 +13,12 @@ struct FDisasterRecoverySession
 	UPROPERTY()
 	FString LastSessionName;
 
+	/** The ID of the process hosting this client session (The PID if of the process for which the transactions are recorded). */
 	UPROPERTY()
-	int32 ProcessId = 0;
+	int32 HostProcessId = 0;
 
 	UPROPERTY()
 	bool bAutoRestoreLastSession = false;
-
-	/** The PID of the last client that write the file for this session. */
-	UPROPERTY()
-	int32 DisasterRecoveryClientPID = 0;
-
-	/** The PID of the disaster recovery service launched by the client for the session. */
-	UPROPERTY()
-	int32 DisasterRecoveryServicePID = 0;
 };
 
 /**
diff --git a/Engine/Source/Programs/CrashReportClient/Private/CrashReportClientApp.cpp b/Engine/Source/Programs/CrashReportClient/Private/CrashReportClientApp.cpp
--- a/Engine/Source/Programs/CrashReportClient/Private/CrashReportClientApp.cpp
+++ b/Engine/Source/Programs/CrashReportClient/Private/CrashReportClientApp.cpp
@@ -612,11 +612,10 @@ void RunCrashReportClient(const TCHAR* CommandLine)
 		double LastTime = FPlatformTime::Seconds();
 		const float IdealFrameTime = 1.0f / IdealFramerate;
 
-		FRecoveryService* RecoveryServicePtr = nullptr;
+		TUniquePtr<FRecoveryService> RecoveryServicePtr;
 #if CRASH_REPORT_WITH_RECOVERY
 		// Starts the disaster recovery service. This records transactions and allows users to recover from previous crashes.
-		FRecoveryService RecoveryService(MonitorPid);
-		RecoveryServicePtr = &RecoveryService;
+		RecoveryServicePtr = MakeUnique<FRecoveryService>(MonitorPid);
 #endif
 
 		FCrashReportAnalytics::Initialize();
@@ -662,7 +661,15 @@ void RunCrashReportClient(const TCHAR* CommandLine)
 				if (IsCrashReportAvailable(MonitorPid, CrashContext, MonitorReadPipe))
 				{
 					// Build error report in memory.
-					FPlatformErrorReport ErrorReport = CollectErrorReport(RecoveryServicePtr, MonitorPid, CrashContext, MonitorWritePipe);
+					FPlatformErrorReport ErrorReport = CollectErrorReport(RecoveryServicePtr.Get(), MonitorPid, CrashContext, MonitorWritePipe);
+
+#if CRASH_REPORT_WITH_RECOVERY
+					if (RecoveryServicePtr && !FPrimaryCrashProperties::Get()->bIsEnsure)
+					{
+						// Shutdown the recovery service, releasing the recovery database file lock (not sharable) as soon as possible to let a new instance take it and offer the user to recover.
+						RecoveryServicePtr.Reset();
+					}
+#endif
 					const SubmitCrashReportResult Result = SendErrorReport(ErrorReport, CrashContext.bNoDialog && CrashContext.bSendUnattenededBugReports);
 
 					// At this point the game can continue execution. It is important this happens
