- Changed 'terminaison' to 'temination'.

#rnx

diff --git a/Engine/Plugins/Developer/Concert/DisasterRecoveryClient/Source/DisasterRecoveryClient/Private/DisasterRecoveryFSM.cpp b/Engine/Plugins/Developer/Concert/DisasterRecoveryClient/Source/DisasterRecoveryClient/Private/DisasterRecoveryFSM.cpp
--- a/Engine/Plugins/Developer/Concert/DisasterRecoveryClient/Source/DisasterRecoveryClient/Private/DisasterRecoveryFSM.cpp
+++ b/Engine/Plugins/Developer/Concert/DisasterRecoveryClient/Source/DisasterRecoveryClient/Private/DisasterRecoveryFSM.cpp
@@ -323,7 +323,7 @@ void FDisasterRecoveryFSM::DisplayRecoveryUI()
 
 	TSharedRef<SConcertSessionRecovery> RecoveryWidget =
 		SNew(SConcertSessionRecovery)
-		.IntroductionText(LOCTEXT("CrashRecoveryIntroductionText", "An abnormal Editor terminaison was detected for this project. You can recover up to the last operation recorded or to a previous state."))
+		.IntroductionText(LOCTEXT("CrashRecoveryIntroductionText", "An abnormal Editor termination was detected for this project. You can recover up to the last operation recorded or to a previous state."))
 		.ParentWindow(NewWindow)
 		.OnFetchActivities(FetchActivitiesFn)
 		.ClientAvatarColorColumnVisibility(EVisibility::Hidden) // Disaster recovery has only one user, the local one.
