#rnx

diff --git a/Engine/Plugins/Developer/Concert/DisasterRecoveryClient/Source/DisasterRecoveryClient/Private/DisasterRecoveryClientModule.cpp b/Engine/Plugins/Developer/Concert/DisasterRecoveryClient/Source/DisasterRecoveryClient/Private/DisasterRecoveryClientModule.cpp
--- a/Engine/Plugins/Developer/Concert/DisasterRecoveryClient/Source/DisasterRecoveryClient/Private/DisasterRecoveryClientModule.cpp
+++ b/Engine/Plugins/Developer/Concert/DisasterRecoveryClient/Source/DisasterRecoveryClient/Private/DisasterRecoveryClientModule.cpp
@@ -316,12 +316,12 @@ private:
 		for (FDisasterRecoverySession& Session : RecoverySessionInfo.Sessions)
 		{
 			const FConcertSessionInfo* ArchivedSession = IsSessionArchived(Session, ArchivedSessions);
-			if (ArchivedSession || IsHostProcessCrashed(Session))
+			if (Session.bAutoRestoreLastSession && (ArchivedSession || IsHostProcessCrashed(Session)))
 			{
 				// Already pinned a session to restore?
 				if (RestoreCandidate.Key && RestoreCandidate.Value)
 				{
-					// User ran multiple instances of the Editor on the same project and more than one instance crashed. Clear this one and restore the first one found only.
+					// User ran multiple instances of the Editor on the same project and more than one instance crashed. Arbitrary clear this one and restore the first one found only.
 					// The situation is unlikely to happen, but a future task would be to return a list of crashed sessions, pass it to the DisasterRecoveryFSM and prompt the user to select which one to recover.
 					Session.bAutoRestoreLastSession = false;
 				}
@@ -362,8 +362,8 @@ private:
 		FDisasterRecoverySessionInfo RecoverySessionInfo;
 		LoadDisasterRecoverySessionInfo(RecoverySessionInfo);
 
-		// If at least one session exist, it may be candidate to restore.
-		return RecoverySessionInfo.Sessions.Num() != 0;
+		// If at least one session exist with the restore flag true, it may be candidate to restore.
+		return RecoverySessionInfo.Sessions.ContainsByPredicate([](const FDisasterRecoverySession& Session) { return Session.bAutoRestoreLastSession; });
 	}
 
 	void ClearSessionInfoFile()
