- Disabled the option to specify a where disaster recovery stores its sessions for 4.24 because it isn't robust enough and the default is good.

diff --git a/Engine/Plugins/Developer/Concert/DisasterRecoveryClient/Source/DisasterRecoveryClient/Private/DisasterRecoveryClientModule.cpp b/Engine/Plugins/Developer/Concert/DisasterRecoveryClient/Source/DisasterRecoveryClient/Private/DisasterRecoveryClientModule.cpp
--- a/Engine/Plugins/Developer/Concert/DisasterRecoveryClient/Source/DisasterRecoveryClient/Private/DisasterRecoveryClientModule.cpp
+++ b/Engine/Plugins/Developer/Concert/DisasterRecoveryClient/Source/DisasterRecoveryClient/Private/DisasterRecoveryClientModule.cpp
@@ -519,8 +519,7 @@ private:
 
 		if (DisasterRecoveryClient)
 		{
-			DisasterRecoveryClient->Shutdown();
-			DisasterRecoveryClient.Reset();
+			StopDisasterRecoveryService();
 		}
 
 		const FString DisasterRecoveryServerName = RecoveryService::GetRecoveryServerName();
diff --git a/Engine/Plugins/Developer/Concert/DisasterRecoveryClient/Source/DisasterRecoveryClient/Private/DisasterRecoverySettings.h b/Engine/Plugins/Developer/Concert/DisasterRecoveryClient/Source/DisasterRecoveryClient/Private/DisasterRecoverySettings.h
--- a/Engine/Plugins/Developer/Concert/DisasterRecoveryClient/Source/DisasterRecoveryClient/Private/DisasterRecoverySettings.h
+++ b/Engine/Plugins/Developer/Concert/DisasterRecoveryClient/Source/DisasterRecoveryClient/Private/DisasterRecoverySettings.h
@@ -24,7 +24,7 @@ public:
 	 * The root directory where recovery sessions should be stored. If not set or
 	 * invalid, the recovery sessions are stored in the project saved directory.
 	 */
-	UPROPERTY(config, EditAnywhere, Category="Client Settings")
+	UPROPERTY(config) // TODO: Expose it once properly tested.
 	FDirectoryPath RecoverySessionDir;
 
 	/**
