Fixed Disaster Recovery keeping the 'auto-connect' the 'discovery mechanism' running in background once the disaster recovery session is created and joined.

Partially addresses:

diff --git a/Engine/Plugins/Developer/Concert/DisasterRecoveryClient/Source/DisasterRecoveryClient/Private/DisasterRecoveryClientModule.cpp b/Engine/Plugins/Developer/Concert/DisasterRecoveryClient/Source/DisasterRecoveryClient/Private/DisasterRecoveryClientModule.cpp
--- a/Engine/Plugins/Developer/Concert/DisasterRecoveryClient/Source/DisasterRecoveryClient/Private/DisasterRecoveryClientModule.cpp
+++ b/Engine/Plugins/Developer/Concert/DisasterRecoveryClient/Source/DisasterRecoveryClient/Private/DisasterRecoveryClientModule.cpp
@@ -437,14 +437,11 @@ private:
 			ClearSessionInfoFile();
 		}
 
-		// Skip restoration if they are no candidates to restore.
-		bool bCandidateToRestore = HasCandidateSessionToRestore();
-
 		// Create and populate the client config object
 		UConcertClientConfig* ClientConfig = NewObject<UConcertClientConfig>();
 		ClientConfig->bIsHeadless = true;
 		ClientConfig->bInstallEditorToolbarButton = false;
-		ClientConfig->bAutoConnect = !bCandidateToRestore; // If there is a possible candidate to restore, don't auto connect yet, run the finite state machine to know more.
+		ClientConfig->bAutoConnect = false;
 		ClientConfig->DefaultServerURL = DisasterRecoveryServerName;
 		ClientConfig->DefaultSessionName = DisasterRecoverySessionName;
 		ClientConfig->DefaultSaveSessionAs = DisasterRecoverySessionName;
@@ -460,22 +457,20 @@ private:
 		SetIgnoreOnRestoreState(!IsCompatibleWithOtherConcertSessions(/*SyncClientStartingSession*/nullptr, /*SyncClientShuttingDownSession*/nullptr));
 
 		// If something might be recovered from a crash.
-		if (bCandidateToRestore)
+		if (HasCandidateSessionToRestore() && GUnrealEd)
 		{
-			if (GUnrealEd)
-			{
-				// Prevent the "Auto-Save" system from restoring the packages before Disaster Recovery plugin.
-				GUnrealEd->GetPackageAutoSaver().DisableRestorePromptAndDeclinePackageRecovery();
-			}
+			// Prevent the "Auto-Save" system from restoring the packages before Disaster Recovery plugin.
+			GUnrealEd->GetPackageAutoSaver().DisableRestorePromptAndDeclinePackageRecovery();
+		}
 
-			// The DisasterRecoveryFSM will call this function to select and take ownership of the session to recover.
-			auto PinSessionFn = [this](const TArray<FConcertSessionInfo>& ArchivedSessions)
-			{
-				return PinSessionToRestore(ArchivedSessions);
-			};
+		// The DisasterRecoveryFSM will call this function to select and take ownership of the session to recover.
+		auto PinSessionFn = [this](const TArray<FConcertSessionInfo>& ArchivedSessions)
+		{
+			return PinSessionToRestore(ArchivedSessions);
+		};
 
-			DisasterRecoveryUtil::StartRecovery(DisasterRecoveryClient.ToSharedRef(), PinSessionFn, /*bLiveDataOnly*/ false);
-		}
+		// The FSM will try to pin a session for recovery (and may fail with a toast), if no session is found, it will create a new one.
+		DisasterRecoveryUtil::StartRecovery(DisasterRecoveryClient.ToSharedRef(), PinSessionFn, /*bLiveDataOnly*/ false);
 
 		return true;
 	}
