Added support to run concurrent disaster recovery clients on the same project. (Enable the user to open the editor twice and get its session transaction recorded)
  - Tracked all client sessions in the 'witness file' rather than a single one.
  - If more than one session ended up crashing, arbitrary pick the first session in the list to recover. (If the user started concurrent editor instance on the same project on his machine and both crashed)
  - Imported the code from dev-enterprise to make the start/stop of the service more robust.

diff --git a/Engine/Plugins/Developer/Concert/DisasterRecoveryClient/Source/DisasterRecoveryClient/Private/DisasterRecoveryClientModule.cpp b/Engine/Plugins/Developer/Concert/DisasterRecoveryClient/Source/DisasterRecoveryClient/Private/DisasterRecoveryClientModule.cpp
--- a/Engine/Plugins/Developer/Concert/DisasterRecoveryClient/Source/DisasterRecoveryClient/Private/DisasterRecoveryClientModule.cpp
+++ b/Engine/Plugins/Developer/Concert/DisasterRecoveryClient/Source/DisasterRecoveryClient/Private/DisasterRecoveryClientModule.cpp
@@ -24,6 +24,7 @@
 #include "UnrealEdGlobals.h"
 #include "Editor/UnrealEdEngine.h"
 #include "IPackageAutoSaver.h"
+#include "GenericPlatform/GenericPlatformCrashContext.h"
 
 #if WITH_EDITOR
 	#include "ISettingsModule.h"
@@ -41,6 +42,8 @@ class FDisasterRecoveryClientModule : public IDisasterRecoveryClientModule
 public:
 	virtual void StartupModule() override
 	{
+		Role = TEXT("DisasterRecovery");
+
 		// Hook to the PreExit callback, needed to execute UObject related shutdowns
 		FCoreDelegates::OnPreExit.AddRaw(this, &FDisasterRecoveryClientModule::HandleAppPreExit);
 
@@ -74,7 +77,7 @@ public:
 private:
 	constexpr bool ShouldLaunchDisasterRecoveryBackendService() const
 	{
-#if PLATFORM_WINDOWS // On windows, the service is launched by the crash analytics process.
+#if PLATFORM_WINDOWS // On windows, the service is launched by the crash analytics process. (WindowsPlatformCrashContext.cpp)
 		return false;
 #else
 		return true;
@@ -107,21 +110,22 @@ private:
 
 	bool HandleSettingsSaved()
 	{
-		const UDisasterRecoverClientConfig* Config = GetDefault<UDisasterRecoverClientConfig>();
-		if (!Config->bIsEnabled)
+		if (GetDefault<UDisasterRecoverClientConfig>()->bIsEnabled)
 		{
-			DisasterRecoveryUtil::EndRecovery(); // Abort recovery.
-			HandleAppPreExit(); // Leave the session and shutdown Concert.
-			DeleteDisasterRecoverySessionInfo(); // Erase the witness file.
-			FPlatformProcess::TerminateProc(DisasterRecoveryServiceHandle); // Required to support restarting. (2 service on the same project/dir fails to share DB files).
+			StartDisasterRecoveryService();
 		}
 		else
 		{
-			OnEngineInitComplete(); // Restart the service/recreate the session.
+			StopDisasterRecoveryService();
 		}
 		return true;
 	}
 
+	void OnEngineInitComplete()
+	{
+		StartDisasterRecoveryService();
+	}
+
 	// Module shutdown is dependent on the UObject system which is currently shutdown on AppExit
 	void HandleAppPreExit()
 	{
@@ -131,84 +135,45 @@ private:
 			return;
 		}
 
-		// Shutdown cleanly - don't auto-restore the active session
-		if (GetDefault<UDisasterRecoverClientConfig>()->bIsEnabled)
-		{
-			FDisasterRecoverySessionInfo SessionInfoToSave;
-			if (LoadDisasterRecoverySessionInfo(SessionInfoToSave))
-			{
-				SessionInfoToSave.bAutoRestoreLastSession = false;
-				SaveDisasterRecoverySessionInfo(SessionInfoToSave);
-			}
-		}
-
-		if (CheckDisasterRecoveryServiceHealthTickHandle.IsValid())
-		{
-			FTicker::GetCoreTicker().RemoveTicker(CheckDisasterRecoveryServiceHealthTickHandle);
-			CheckDisasterRecoveryServiceHealthTickHandle.Reset();
-		}
-
-		if (DisasterRecoveryClient)
-		{
-			DisasterRecoveryClient->Shutdown();
-			DisasterRecoveryClient.Reset();
-		}
+		StopDisasterRecoveryService();
 	}
 
-	void OnEngineInitComplete()
+	FGuid GetDisasterRecoverySessionId() const
 	{
-		const UDisasterRecoverClientConfig* Config = GetDefault<UDisasterRecoverClientConfig>();
-		if (!Config->bIsEnabled)
-		{
-			return;
-		}
-
+		if (DisasterRecoveryClient)
 		{
-			FDisasterRecoverySessionInfo SessionInfoToRestore;
-			if (LoadDisasterRecoverySessionInfo(SessionInfoToRestore))
+			if (TSharedPtr<IConcertClientSession> Session = DisasterRecoveryClient->GetConcertClient()->GetCurrentSession())
 			{
-				StartDisasterRecoveryService(&SessionInfoToRestore);
-			}
-			else
-			{
-				StartDisasterRecoveryService();
+				return Session->GetSessionInfo().SessionId;
 			}
 		}
 
-		if (ShouldLaunchDisasterRecoveryBackendService())
-		{
-			CheckDisasterRecoveryServiceHealthTickHandle = FTicker::GetCoreTicker().AddTicker(TEXT("CheckDisasterRecoveryServiceHealth"), 1.0f, [this](float)
-			{
-				CheckDisasterRecoveryServiceHealth();
-				return true;
-			});
-		}
+		return FGuid(); // Invalid.
 	}
 
 	/** Returns the folder where the disaster recovery service should keep the live session files (the working directory). */
-	static FString GetDefaultServerWorkingDir()
+	FString GetDefaultServerWorkingDir() const
 	{
-		return FPaths::ProjectIntermediateDir() / TEXT("Concert") / TEXT("Server");
+		return FPaths::ProjectIntermediateDir() / Role / TEXT("Service");
 	}
 
 	/** Returns the folder where the disaster recovery service should keep the archived session files (the saved directory). */
-	static FString GetDefaultServerArchiveDir()
+	FString GetDefaultServerArchiveDir() const
 	{
 		// Put the session data in the project dir.
-		return FPaths::ProjectSavedDir() / TEXT("Concert") / TEXT("Server");
+		return FPaths::ProjectSavedDir() / Role / TEXT("Service");
 	}
 
-	static FString GetDisasterRecoverySessionInfoFilename()
+	FString GetDisasterRecoverySessionInfoFilename() const
 	{
-		return FPaths::ProjectSavedDir() / TEXT("Concert") / TEXT("DisasterRecovery") / TEXT("SessionInfo.json");
+		return FPaths::ProjectSavedDir() / Role / TEXT("SessionInfo.json");
 	}
 
-	static bool LoadDisasterRecoverySessionInfo(FDisasterRecoverySessionInfo& OutSessionInfo)
+	bool LoadDisasterRecoverySessionInfo(FDisasterRecoverySessionInfo& OutSessionInfo)
 	{
 		if (TUniquePtr<FArchive> FileReader = TUniquePtr<FArchive>(IFileManager::Get().CreateFileReader(*GetDisasterRecoverySessionInfoFilename())))
 		{
 			FJsonStructDeserializerBackend Backend(*FileReader);
-
 			FStructDeserializer::Deserialize(OutSessionInfo, Backend);
 
 			FileReader->Close();
@@ -218,12 +183,11 @@ private:
 		return false;
 	}
 
-	static bool SaveDisasterRecoverySessionInfo(const FDisasterRecoverySessionInfo& InSessionInfo)
+	bool SaveDisasterRecoverySessionInfo(const FDisasterRecoverySessionInfo& InSessionInfo)
 	{
 		if (TUniquePtr<FArchive> FileWriter = TUniquePtr<FArchive>(IFileManager::Get().CreateFileWriter(*GetDisasterRecoverySessionInfoFilename())))
 		{
 			FJsonStructSerializerBackend Backend(*FileWriter, EStructSerializerBackendFlags::Default);
-
 			FStructSerializer::Serialize(InSessionInfo, Backend);
 
 			FileWriter->Close();
@@ -233,11 +197,6 @@ private:
 		return false;
 	}
 
-	static bool DeleteDisasterRecoverySessionInfo()
-	{
-		return IFileManager::Get().Delete(*GetDisasterRecoverySessionInfoFilename());
-	}
-
 	static FString GetDisasterRecoveryServicePath()
 	{
 		auto GetDisasterRecoveryServicePathForBuildConfiguration = [](const EBuildConfiguration InBuildConfiguration) -> FString
@@ -258,80 +217,168 @@ private:
 		return ServicePath;
 	}
 
+	static FString GetSystemMutexName()
+	{
+		return TEXT("Unreal_DisasterRecovery_4221FF"); // Arbitrary name that is unique among other applications.
+	}
+
 	void DisasterRecoverySessionCreated(TSharedRef<IConcertClientSession> InSession)
 	{
-		FDisasterRecoverySessionInfo SessionInfoToSave;
-		SessionInfoToSave.LastSessionName = InSession->GetSessionInfo().SessionName;
+		// Ensure we get exclusive access to the file.
+		FSystemWideCriticalSection SystemWideMutex(GetSystemMutexName());
+
+		// Reload the file.
+		FDisasterRecoverySessionInfo RecoverySessionInfo;
+		LoadDisasterRecoverySessionInfo(RecoverySessionInfo);
+
+		// Is this a new session created by recovering from another one?
+		if (RestoringSession.IsSet())
+		{
+			// Remove the recovered session from the list.
+			RecoverySessionInfo.Sessions.RemoveAll([this](const FDisasterRecoverySession& Session)
+			{
+				return Session.bAutoRestoreLastSession && RestoringSession->ProcessId == Session.ProcessId && RestoringSession->LastSessionName == Session.LastSessionName;
+			});
+
+			RestoringSession.Reset();
+		}
+
+		// Create a new session.
+		FDisasterRecoverySession& Session = RecoverySessionInfo.Sessions.AddDefaulted_GetRef();
+		Session.LastSessionName = InSession->GetSessionInfo().SessionName;
+		Session.ProcessId = FPlatformProcess::GetCurrentProcessId();
 
 		// Normally, bAutoRestoreLastSession is set true here and overwritten to false when the app exits normally, but when running under the debugger, auto-restore is disabled as
 		// programmers kill applications (stop the debugger) and this should not count as a crash (unless you want to simulate crash this way - see below).
-		SessionInfoToSave.bAutoRestoreLastSession = !FPlatformMisc::IsDebuggerPresent();
-		//SessionInfoToSave.bAutoRestoreLastSession = true; // <- THIS CODE MUST BE COMMENTED BEFORE SUBMITTING: Uncomment for debugging purpose only. Simulate a crash by stopping the debugger during a session.
-		SaveDisasterRecoverySessionInfo(SessionInfoToSave);
+		Session.bAutoRestoreLastSession = !FPlatformMisc::IsDebuggerPresent();
+		//Sessions.bAutoRestoreLastSession = true; // <- MUST BE COMMENTED BEFORE SUBMITTING: For debugging purpose only. Simulate a crash by stopping the debugger during a session.
+
+		// Save the file.
+		SaveDisasterRecoverySessionInfo(RecoverySessionInfo);
 	}
 
-	bool StartDisasterRecoveryService(const FDisasterRecoverySessionInfo* InSessionInfoToRestore = nullptr)
+	TOptional<FDisasterRecoverySession> ProcessSessionInfo()
 	{
-		if (DisasterRecoveryClient)
+		// Ensure we get exclusive access to the file.
+		FSystemWideCriticalSection SystemWideMutex(GetSystemMutexName());
+
+		TOptional<FDisasterRecoverySession> SessionToRestore;
+
+		// Load the session info file (if it exist)
+		FDisasterRecoverySessionInfo RecoverySessionInfo;
+		LoadDisasterRecoverySessionInfo(RecoverySessionInfo);
+
+		for (FDisasterRecoverySession& Session : RecoverySessionInfo.Sessions)
 		{
-			DisasterRecoveryClient->Shutdown();
-			DisasterRecoveryClient.Reset();
+			if (Session.bAutoRestoreLastSession)
+			{
+				if (!FPlatformProcess::IsApplicationRunning(Session.ProcessId) || // If the session owner is not running anymore.
+					FPaths::GetPathLeaf(FPlatformProcess::GetApplicationName(Session.ProcessId)) != FPaths::GetPathLeaf(FPlatformProcess::GetApplicationName(FPlatformProcess::GetCurrentProcessId()))) // This owner PID is live, but was reused by the OS for another app.
+				{
+					if (!SessionToRestore.IsSet())
+					{
+						// Recover this session. Take ownership of this session. (Will be cleared on DisasterRecoverySessionCreated() if recovery succeed)
+						Session.ProcessId = FPlatformProcess::GetCurrentProcessId();
+
+						// Pick up the first crashed session. (This is arbitrary if they are more)
+						SessionToRestore.Emplace(Session);
+					}
+					else
+					{
+						// User ran multiple instances of the Editor on the same project and more than one instance crashed. Just restore the first one found above for the moment.
+						// While the situation above is unlikely to happen, a future task would be to return a list of crashed session, pass it to the DisasterRecoveryFSM and prompt the user to which one to recover.
+						Session.bAutoRestoreLastSession = false;
+					}
+				}
+			}
 		}
 
-		if (!FApp::HasProjectName())
+		// Remove completed sessions.
+		RecoverySessionInfo.Sessions.RemoveAll([](const FDisasterRecoverySession& Session)
+		{
+			return !Session.bAutoRestoreLastSession; // Remove if the 'restore' flag is false.
+		});
+
+		// Save the recovery session info file.
+		SaveDisasterRecoverySessionInfo(RecoverySessionInfo);
+
+		return SessionToRestore;
+	}
+
+	bool SpawnDisasterRecoveryServer(const FString& ServerName) // On Linux and Mac. On Windows, it is embedded in CrashReporterClient.
+	{
+		// Find the service path that will host the sync server
+		const FString DisasterRecoveryServicePath = GetDisasterRecoveryServicePath();
+		if (DisasterRecoveryServicePath.IsEmpty())
 		{
+			UE_LOG(LogDisasterRecovery, Warning, TEXT("Disaster Recovery Service application was not found. Disaster Recovery will be disabled! Please build 'UnrealDisasterRecoveryService'."));
 			return false;
 		}
 
-		const FString DisasterRecoveryServerName = FString::Printf(TEXT("RecoverySvr_%d"), FPlatformProcess::GetCurrentProcessId());
-		const FString DisasterRecoverySessionName = FString::Printf(TEXT("Recovery_%d_%s_%s"), FPlatformProcess::GetCurrentProcessId(), FApp::GetProjectName(), *FDateTime::Now().ToString());
+		FString DisasterRecoveryServiceCommandLine;
+		DisasterRecoveryServiceCommandLine += FString::Printf(TEXT(" -ConcertServer=\"%s\""), *ServerName);
+		DisasterRecoveryServiceCommandLine += FString::Printf(TEXT(" -EditorPID=%d"), FPlatformProcess::GetCurrentProcessId());
+		DisasterRecoveryServiceCommandLine += FString::Printf(TEXT(" -ConcertWorkingDir=\"%s\""), *GetDefaultServerWorkingDir());
+		DisasterRecoveryServiceCommandLine += FString::Printf(TEXT(" -ConcertSavedDir=\"%s\""), *GetDefaultServerArchiveDir());
 
-		if (ShouldLaunchDisasterRecoveryBackendService())
+		// Create the service process that will host the sync server
+		DisasterRecoveryServiceHandle = FPlatformProcess::CreateProc(*DisasterRecoveryServicePath, *DisasterRecoveryServiceCommandLine, true, true, true, nullptr, 0, nullptr, nullptr, nullptr);
+		if (!DisasterRecoveryServiceHandle.IsValid())
 		{
-			// Find the service path that will host the sync server
-			const FString DisasterRecoveryServicePath = GetDisasterRecoveryServicePath();
-			if (DisasterRecoveryServicePath.IsEmpty())
-			{
-				UE_LOG(LogDisasterRecovery, Warning, TEXT("Disaster Recovery Service application was not found. Disaster Recovery will be disabled! Please build 'UnrealDisasterRecoveryService'."));
-				return false;
-			}
+			UE_LOG(LogDisasterRecovery, Error, TEXT("Failed to launch Disaster Recovery Service application. Disaster Recovery will be disabled!"));
+			return false;
+		}
 
-			FString DisasterRecoveryServiceCommandLine;
-			DisasterRecoveryServiceCommandLine += FString::Printf(TEXT(" -ConcertServer=\"%s\""), *DisasterRecoveryServerName);
-			DisasterRecoveryServiceCommandLine += TEXT(" -ConcertIgnore");
-			DisasterRecoveryServiceCommandLine += FString::Printf(TEXT(" -EditorPID=%d"), FPlatformProcess::GetCurrentProcessId());
-			DisasterRecoveryServiceCommandLine += FString::Printf(TEXT(" -ConcertWorkingDir=\"%s\""), *GetDefaultServerWorkingDir());
-			DisasterRecoveryServiceCommandLine += FString::Printf(TEXT(" -ConcertSavedDir=\"%s\""), *GetDefaultServerArchiveDir());
+		return true;
+	}
 
-			// Create the service process that will host the sync server
-			DisasterRecoveryServiceHandle = FPlatformProcess::CreateProc(*DisasterRecoveryServicePath, *DisasterRecoveryServiceCommandLine, true, true, true, nullptr, 0, nullptr, nullptr, nullptr);
-			if (!DisasterRecoveryServiceHandle.IsValid())
-			{
-				UE_LOG(LogDisasterRecovery, Error, TEXT("Failed to launch Disaster Recovery Service application. Disaster Recovery will be disabled!"));
-				return false;
-			}
+	bool StartDisasterRecoveryService()
+	{
+		if (!GetDefault<UDisasterRecoverClientConfig>()->bIsEnabled)
+		{
+			return false;
 		}
 
-		bool bRecoveringFromCrash = InSessionInfoToRestore && InSessionInfoToRestore->bAutoRestoreLastSession;
+		if (!FApp::HasProjectName())
+		{
+			return false;
+		}
+
+		if (DisasterRecoveryClient)
+		{
+			DisasterRecoveryClient->Shutdown();
+			DisasterRecoveryClient.Reset();
+		}
+
+		const FString DisasterRecoveryServerName = RecoveryService::GetRecoveryServerName();
+		const FString DisasterRecoverySessionName = FString::Printf(TEXT("%s_%s_%s"), FApp::GetProjectName(), *FDateTime::Now().ToString(), *DisasterRecoveryServerName);
+
+		if (ShouldLaunchDisasterRecoveryBackendService() && !SpawnDisasterRecoveryServer(DisasterRecoveryServerName))
+		{
+			return false; // Failed to spawn the service.
+		}
+
+		// Clean sessions and find if a session should be restored
+		RestoringSession = ProcessSessionInfo();
 
 		// Create and populate the client config object
 		UConcertClientConfig* ClientConfig = NewObject<UConcertClientConfig>();
 		ClientConfig->bIsHeadless = true;
 		ClientConfig->bInstallEditorToolbarButton = false;
-		ClientConfig->bAutoConnect = !bRecoveringFromCrash; // If recovering from a crash, don't auto connect -> Present UI to let user decide what to recover first.
+		ClientConfig->bAutoConnect = !RestoringSession.IsSet(); // If recovering from a crash, don't auto connect -> Present UI to let user decide what to recover first.
 		ClientConfig->DefaultServerURL = DisasterRecoveryServerName;
 		ClientConfig->DefaultSessionName = DisasterRecoverySessionName;
 		ClientConfig->DefaultSaveSessionAs = DisasterRecoverySessionName;
 		//ClientConfig->ClientSettings.DiscoveryTimeoutSeconds = 0;
 		ClientConfig->EndpointSettings.RemoteEndpointTimeoutSeconds = 0;
 
-		// Create and boot the client instance, auto-joining the configured session if launching normally (not in recovery mode).
-		DisasterRecoveryClient = IConcertSyncClientModule::Get().CreateClient(TEXT("DisasterRecovery"));
+		// Create the recovery session and auto-join it if there is nothing to recover.
+		DisasterRecoveryClient = IConcertSyncClientModule::Get().CreateClient(Role);
 		DisasterRecoveryClient->GetConcertClient()->OnSessionStartup().AddRaw(this, &FDisasterRecoveryClientModule::DisasterRecoverySessionCreated);
 		DisasterRecoveryClient->Startup(ClientConfig, EConcertSyncSessionFlags::Default_DisasterRecoverySession);
 
-		// Starts the recovery process.
-		if (bRecoveringFromCrash)
+		// If something needs to be recovered from a crash.
+		if (RestoringSession)
 		{
 			if (GUnrealEd)
 			{
@@ -339,41 +386,63 @@ private:
 				GUnrealEd->GetPackageAutoSaver().DisableRestorePromptAndDeclinePackageRecovery();
 			}
 
-			DisasterRecoveryUtil::StartRecovery(DisasterRecoveryClient.ToSharedRef(), InSessionInfoToRestore->LastSessionName, /*bLiveDataOnly*/ false);
+			DisasterRecoveryUtil::StartRecovery(DisasterRecoveryClient.ToSharedRef(), RestoringSession->LastSessionName, /*bLiveDataOnly*/ false);
 		}
 
 		return true;
 	}
 
-	void CheckDisasterRecoveryServiceHealth()
+	void StopDisasterRecoveryService()
 	{
-		if (!DisasterRecoveryServiceHandle.IsValid())
+		// End the recovery FSM (if running).
+		bool bRecoveryAborted = !DisasterRecoveryUtil::EndRecovery();
+		if (!bRecoveryAborted) // Can be aborted if the user close the editor before the recovery modal window appears. (Need to be quick, but possible)
 		{
-			// Disaster Recovery is not available - nothing to check the health of
-			return;
+			FSystemWideCriticalSection SystemWideMutex(GetSystemMutexName());
+
+			// Load the sessions file.
+			FDisasterRecoverySessionInfo RecoverySessionInfo;
+			LoadDisasterRecoverySessionInfo(RecoverySessionInfo);
+
+			// Remove the current session from the list of sessions to track.
+			int32 ProcessId = FPlatformProcess::GetCurrentProcessId();
+			RecoverySessionInfo.Sessions.RemoveAll([ProcessId](const FDisasterRecoverySession& Session)
+			{
+				return Session.ProcessId == ProcessId;
+			});
+
+			// Write the file to disk.
+			SaveDisasterRecoverySessionInfo(RecoverySessionInfo);
 		}
 
-		// If the Disaster Recovery Service stopped, attempt to restart it
-		if (!FPlatformProcess::IsProcRunning(DisasterRecoveryServiceHandle))
+		if (DisasterRecoveryClient)
 		{
+			DisasterRecoveryClient->Shutdown();
+			DisasterRecoveryClient.Reset();
+		}
+
+		if (DisasterRecoveryServiceHandle.IsValid())
+		{
+			FPlatformProcess::TerminateProc(DisasterRecoveryServiceHandle);
 			DisasterRecoveryServiceHandle.Reset();
-			StartDisasterRecoveryService();
 		}
 	}
 
+private:
+	/** This client role, a tag given to different types of concert client, i.e. DisasterRecovery for this one. */
+	FString Role;
+
 	/** Sync client handling disaster recovery */
 	TSharedPtr<IConcertSyncClient> DisasterRecoveryClient;
 
 	/** Handle to the active disaster recovery service app, if any */
 	FProcHandle DisasterRecoveryServiceHandle;
 
-	/** Ticker for ensuring that the disaster recovery service app that was started is still running */
-	FDelegateHandle CheckDisasterRecoveryServiceHealthTickHandle;
-
-	/** Count of the number of times the disaster recovery service app has been started by this client */
-	int32 DisasterRecoveryServiceInstanceCount = 0;
+	/** Keep the session being restored, if any. */
+	TOptional<FDisasterRecoverySession> RestoringSession;
 };
 
 #undef LOCTEXT_NAMESPACE
 
 IMPLEMENT_MODULE(FDisasterRecoveryClientModule, DisasterRecoveryClient);
+
diff --git a/Engine/Plugins/Developer/Concert/DisasterRecoveryClient/Source/DisasterRecoveryClient/Private/DisasterRecoveryFSM.cpp b/Engine/Plugins/Developer/Concert/DisasterRecoveryClient/Source/DisasterRecoveryClient/Private/DisasterRecoveryFSM.cpp
--- a/Engine/Plugins/Developer/Concert/DisasterRecoveryClient/Source/DisasterRecoveryClient/Private/DisasterRecoveryFSM.cpp
+++ b/Engine/Plugins/Developer/Concert/DisasterRecoveryClient/Source/DisasterRecoveryClient/Private/DisasterRecoveryFSM.cpp
@@ -31,6 +31,11 @@ public:
 	FDisasterRecoveryFSM(TSharedRef<IConcertSyncClient> InSyncClient, const FString& InSessionNameToRecover, bool bInLiveDataOnly);
 	~FDisasterRecoveryFSM();
 
+	bool IsDone() const
+	{
+		return CurrentState == &ExitState;
+	}
+
 private:
 	struct FState
 	{
@@ -435,9 +440,16 @@ void DisasterRecoveryUtil::StartRecovery(TSharedRef<IConcertSyncClient> SyncClie
 	}
 }
 
-void DisasterRecoveryUtil::EndRecovery()
+bool DisasterRecoveryUtil::EndRecovery()
 {
-	RecoveryFSM.Reset();
+	bool bDone = true;
+	if (RecoveryFSM)
+	{
+		bDone = RecoveryFSM->IsDone();
+		RecoveryFSM.Reset();
+	}
+
+	return bDone;
 }
 
 #undef LOCTEXT_NAMESPACE
diff --git a/Engine/Plugins/Developer/Concert/DisasterRecoveryClient/Source/DisasterRecoveryClient/Private/DisasterRecoveryFSM.h b/Engine/Plugins/Developer/Concert/DisasterRecoveryClient/Source/DisasterRecoveryClient/Private/DisasterRecoveryFSM.h
--- a/Engine/Plugins/Developer/Concert/DisasterRecoveryClient/Source/DisasterRecoveryClient/Private/DisasterRecoveryFSM.h
+++ b/Engine/Plugins/Developer/Concert/DisasterRecoveryClient/Source/DisasterRecoveryClient/Private/DisasterRecoveryFSM.h
@@ -18,7 +18,10 @@ namespace DisasterRecoveryUtil
  */
 void StartRecovery(TSharedRef<IConcertSyncClient> SyncClient, const FString& SessionNameToRecover, bool bLiveDataOnly);
 
-/** End the recovery flow. This can be called to abort the recovery process. */
-void EndRecovery();
+/**
+ * End the recovery flow. This can be called to abort the recovery process.
+ * @return True if the recovery already completed successfully, false if it was not completed and it was aborted.
+ */
+bool EndRecovery();
 
-}
\ No newline at end of file
+}
diff --git a/Engine/Plugins/Developer/Concert/DisasterRecoveryClient/Source/DisasterRecoveryClient/Private/DisasterRecoverySessionInfo.h b/Engine/Plugins/Developer/Concert/DisasterRecoveryClient/Source/DisasterRecoveryClient/Private/DisasterRecoverySessionInfo.h
--- a/Engine/Plugins/Developer/Concert/DisasterRecoveryClient/Source/DisasterRecoveryClient/Private/DisasterRecoverySessionInfo.h
+++ b/Engine/Plugins/Developer/Concert/DisasterRecoveryClient/Source/DisasterRecoveryClient/Private/DisasterRecoverySessionInfo.h
@@ -6,13 +6,28 @@
 #include "DisasterRecoverySessionInfo.generated.h"
 
 USTRUCT()
-struct FDisasterRecoverySessionInfo
+struct FDisasterRecoverySession
 {
 	GENERATED_BODY()
 
 	UPROPERTY()
 	FString LastSessionName;
 
+	UPROPERTY()
+	int32 ProcessId = 0;
+
 	UPROPERTY()
 	bool bAutoRestoreLastSession = false;
 };
+
+/**
+ * Hold the information for multiple disaster recovery sessions.
+ */
+USTRUCT()
+struct FDisasterRecoverySessionInfo
+{
+	GENERATED_BODY()
+
+	UPROPERTY()
+	TArray<FDisasterRecoverySession> Sessions;
+};
