Prevented disaster recovery from reapplying operations performed during a multi-user session.

diff --git a/Engine/Plugins/Developer/Concert/ConcertSync/ConcertSyncClient/Source/ConcertSyncClient/Private/ConcertClientPackageManager.cpp b/Engine/Plugins/Developer/Concert/ConcertSync/ConcertSyncClient/Source/ConcertSyncClient/Private/ConcertClientPackageManager.cpp
--- a/Engine/Plugins/Developer/Concert/ConcertSync/ConcertSyncClient/Source/ConcertSyncClient/Private/ConcertClientPackageManager.cpp
+++ b/Engine/Plugins/Developer/Concert/ConcertSync/ConcertSyncClient/Source/ConcertSyncClient/Private/ConcertClientPackageManager.cpp
@@ -177,6 +177,7 @@ void FConcertClientPackageManager::HandlePackageDiscarded(UPackage* InPackage)
 	Event.Package.Info.PackageName = InPackage->GetFName();
 	Event.Package.Info.PackageFileExtension = UWorld::FindWorldInPackage(InPackage) ? FPackageName::GetMapPackageExtension() : FPackageName::GetAssetPackageExtension();
 	Event.Package.Info.PackageUpdateType = EConcertPackageUpdateType::Dummy;
+	Event.bReplayable = bReplayableEvents;
 	LiveSession->GetSessionDatabase().GetTransactionMaxEventId(Event.Package.Info.TransactionEventIdAtSave);
 	LiveSession->GetSession().SendCustomEvent(Event, LiveSession->GetSession().GetSessionServerEndpointId(), EConcertMessageFlags::ReliableOrdered);
 }
@@ -329,6 +330,7 @@ void FConcertClientPackageManager::HandleLocalPackageEvent(const FConcertPackage
 
 	FConcertPackageUpdateEvent Event;
 	Event.Package = Package;
+	Event.bReplayable = bReplayableEvents;
 	LiveSession->GetSessionDatabase().GetTransactionMaxEventId(Event.Package.Info.TransactionEventIdAtSave);
 	LiveSession->GetSession().SendCustomEvent(Event, LiveSession->GetSession().GetSessionServerEndpointId(), EConcertMessageFlags::ReliableOrdered);
 }
diff --git a/Engine/Plugins/Developer/Concert/ConcertSync/ConcertSyncClient/Source/ConcertSyncClient/Private/ConcertClientPackageManager.h b/Engine/Plugins/Developer/Concert/ConcertSync/ConcertSyncClient/Source/ConcertSyncClient/Private/ConcertClientPackageManager.h
--- a/Engine/Plugins/Developer/Concert/ConcertSync/ConcertSyncClient/Source/ConcertSyncClient/Private/ConcertClientPackageManager.h
+++ b/Engine/Plugins/Developer/Concert/ConcertSync/ConcertSyncClient/Source/ConcertSyncClient/Private/ConcertClientPackageManager.h
@@ -77,6 +77,11 @@ public:
 	 */
 	bool PersistSessionChanges(TArrayView<const FString> InFilesToPersist, ISourceControlProvider* SourceControlProvider, TArray<FText>* OutFailureReasons = nullptr);
 
+	/**
+	 * Sets whether further package events emitted to the server has the 'replayable' flag on or off.
+	 */
+	void SetPackageEventAsReplayable(bool bReplayable) { bReplayableEvents = bReplayable; }
+
 private:
 	/**
 	 * Apply the data in the given package to disk and update the in-memory state.
@@ -163,4 +168,9 @@ private:
 	 * Array of package names that are pending an in-memory purge.
 	 */
 	TArray<FName> PackagesPendingPurge;
+
+	/**
+	 * Whether the package events emitted to the server has the replayable flag on or off.
+	 */
+	bool bReplayableEvents = false;
 };
diff --git a/Engine/Plugins/Developer/Concert/ConcertSync/ConcertSyncClient/Source/ConcertSyncClient/Private/ConcertClientTransactionManager.cpp b/Engine/Plugins/Developer/Concert/ConcertSync/ConcertSyncClient/Source/ConcertSyncClient/Private/ConcertClientTransactionManager.cpp
--- a/Engine/Plugins/Developer/Concert/ConcertSync/ConcertSyncClient/Source/ConcertSyncClient/Private/ConcertClientTransactionManager.cpp
+++ b/Engine/Plugins/Developer/Concert/ConcertSync/ConcertSyncClient/Source/ConcertSyncClient/Private/ConcertClientTransactionManager.cpp
@@ -389,13 +389,14 @@ void FConcertClientTransactionManager::HandleLocalTransactionFinalized(const FCo
 			return;
 		}
 	}
-	
+
 	FPendingTransactionToSend& PendingTransaction = HandleLocalTransactionCommon(InCommonData);
+	PendingTransaction.CommonData.bReplayable = bReplayableEvents;
 	PendingTransaction.FinalizedData = InFinalizedData;
 	PendingTransaction.bIsFinalized = true;
 }
 
-void FConcertClientTransactionManager::SendTransactionFinalizedEvent(const FGuid& InTransactionId, const FGuid& InOperationId, UObject* InPrimaryObject, const TArray<FName>& InModifiedPackages, const TArray<FConcertExportedObject>& InObjectUpdates, const FConcertLocalIdentifierTable& InLocalIdentifierTable, const FText& InTitle)
+void FConcertClientTransactionManager::SendTransactionFinalizedEvent(const FGuid& InTransactionId, const FGuid& InOperationId, UObject* InPrimaryObject, const TArray<FName>& InModifiedPackages, const TArray<FConcertExportedObject>& InObjectUpdates, const FConcertLocalIdentifierTable& InLocalIdentifierTable, const FText& InTitle, bool bReplayable)
 {
 	FConcertTransactionFinalizedEvent TransactionFinalizedEvent;
 	FillTransactionEvent(InTransactionId, InOperationId, InModifiedPackages, TransactionFinalizedEvent);
@@ -403,6 +404,7 @@ void FConcertClientTransactionManager::SendTransactionFinalizedEvent(const FGuid
 	TransactionFinalizedEvent.ExportedObjects = InObjectUpdates;
 	InLocalIdentifierTable.GetState(TransactionFinalizedEvent.LocalIdentifierState);
 	TransactionFinalizedEvent.Title = InTitle;
+	TransactionFinalizedEvent.bReplayable = bReplayable;
 
 	LiveSession->GetSession().SendCustomEvent(TransactionFinalizedEvent, LiveSession->GetSession().GetSessionServerEndpointId(), EConcertMessageFlags::ReliableOrdered);
 }
@@ -451,7 +453,8 @@ void FConcertClientTransactionManager::SendPendingTransactionEvents()
 						PendingTransactionPtr->CommonData.ModifiedPackages, 
 						PendingTransactionPtr->FinalizedData.FinalizedObjectUpdates, 
 						PendingTransactionPtr->FinalizedData.FinalizedLocalIdentifierTable, 
-						PendingTransactionPtr->CommonData.TransactionTitle
+						PendingTransactionPtr->CommonData.TransactionTitle,
+						PendingTransactionPtr->CommonData.bReplayable
 						);
 				}
 				// TODO: Warn about excluded objects?
diff --git a/Engine/Plugins/Developer/Concert/ConcertSync/ConcertSyncClient/Source/ConcertSyncClient/Private/ConcertClientTransactionManager.h b/Engine/Plugins/Developer/Concert/ConcertSync/ConcertSyncClient/Source/ConcertSyncClient/Private/ConcertClientTransactionManager.h
--- a/Engine/Plugins/Developer/Concert/ConcertSync/ConcertSyncClient/Source/ConcertSyncClient/Private/ConcertClientTransactionManager.h
+++ b/Engine/Plugins/Developer/Concert/ConcertSync/ConcertSyncClient/Source/ConcertSyncClient/Private/ConcertClientTransactionManager.h
@@ -52,6 +52,11 @@ public:
 	 */
 	void ProcessPending();
 
+	/**
+	 * Sets whether further events emitted to the server has the 'replayable' flag on or off.
+	 */
+	void SetTransactionEventsAsReplayable(bool bReplayable) { bReplayableEvents = bReplayable; }
+
 private:
 	/**
 	 * Context object for transactions that are to be processed.
@@ -155,7 +160,7 @@ private:
 	/**
 	 * Send a transaction finalized event.
 	 */
-	void SendTransactionFinalizedEvent(const FGuid& InTransactionId, const FGuid& InOperationId, UObject* InPrimaryObject, const TArray<FName>& InModifiedPackages, const TArray<FConcertExportedObject>& InObjectUpdates, const FConcertLocalIdentifierTable& InLocalIdentifierTable, const FText& InTitle);
+	void SendTransactionFinalizedEvent(const FGuid& InTransactionId, const FGuid& InOperationId, UObject* InPrimaryObject, const TArray<FName>& InModifiedPackages, const TArray<FConcertExportedObject>& InObjectUpdates, const FConcertLocalIdentifierTable& InLocalIdentifierTable, const FText& InTitle, bool bReplayable);
 
 	/**
 	 * Send a transaction snapshot event.
@@ -203,4 +208,9 @@ private:
 	 * Map of transaction IDs to the pending transaction that may be sent in the future (when finalized).
 	 */
 	TMap<FGuid, FPendingTransactionToSend> PendingTransactionsToSend;
+
+	/**
+	 * Whether the transaction events emitted to the server has the replayable flag on or off.
+	 */
+	bool bReplayableEvents = true;
 };
diff --git a/Engine/Plugins/Developer/Concert/ConcertSync/ConcertSyncClient/Source/ConcertSyncClient/Private/ConcertClientWorkspace.cpp b/Engine/Plugins/Developer/Concert/ConcertSync/ConcertSyncClient/Source/ConcertSyncClient/Private/ConcertClientWorkspace.cpp
--- a/Engine/Plugins/Developer/Concert/ConcertSync/ConcertSyncClient/Source/ConcertSyncClient/Private/ConcertClientWorkspace.cpp
+++ b/Engine/Plugins/Developer/Concert/ConcertSync/ConcertSyncClient/Source/ConcertSyncClient/Private/ConcertClientWorkspace.cpp
@@ -763,4 +763,12 @@ void FConcertClientWorkspace::PostActivityUpdated(const FConcertSyncActivity& In
 	}
 }
 
+void FConcertClientWorkspace::SetEmittedEventsAsReplayable(bool bReplayable)
+{
+	check(TransactionManager && PackageManager);
+
+	TransactionManager->SetTransactionEventsAsReplayable(bReplayable);
+	PackageManager->SetPackageEventAsReplayable(bReplayable);
+}
+
 #undef LOCTEXT_NAMESPACE
diff --git a/Engine/Plugins/Developer/Concert/ConcertSync/ConcertSyncClient/Source/ConcertSyncClient/Private/ConcertClientWorkspace.h b/Engine/Plugins/Developer/Concert/ConcertSync/ConcertSyncClient/Source/ConcertSyncClient/Private/ConcertClientWorkspace.h
--- a/Engine/Plugins/Developer/Concert/ConcertSync/ConcertSyncClient/Source/ConcertSyncClient/Private/ConcertClientWorkspace.h
+++ b/Engine/Plugins/Developer/Concert/ConcertSync/ConcertSyncClient/Source/ConcertSyncClient/Private/ConcertClientWorkspace.h
@@ -49,6 +49,7 @@ public:
 	virtual FOnWorkspaceSynchronized& OnWorkspaceSynchronized() override;
 	virtual IConcertClientDataStore& GetDataStore() override;
 	virtual bool IsAssetModifiedByOtherClients(const FName& AssetName, int32* OutOtherClientsWithModifNum, TArray<FConcertClientInfo>* OutOtherClientsWithModifInfo, int32 OtherClientsWithModifMaxFetchNum) const override;
+	virtual void SetEmittedEventsAsReplayable(bool bEnabled) override;
 
 private:
 	/** Bind the workspace to this session. */
diff --git a/Engine/Plugins/Developer/Concert/ConcertSync/ConcertSyncClient/Source/ConcertSyncClient/Private/ConcertSyncClient.cpp b/Engine/Plugins/Developer/Concert/ConcertSync/ConcertSyncClient/Source/ConcertSyncClient/Private/ConcertSyncClient.cpp
--- a/Engine/Plugins/Developer/Concert/ConcertSync/ConcertSyncClient/Source/ConcertSyncClient/Private/ConcertSyncClient.cpp
+++ b/Engine/Plugins/Developer/Concert/ConcertSync/ConcertSyncClient/Source/ConcertSyncClient/Private/ConcertSyncClient.cpp
@@ -94,6 +94,16 @@ FOnConcertClientWorkspaceStartupOrShutdown& FConcertSyncClient::OnWorkspaceShutd
 	return OnWorkspaceShutdownDelegate;
 }
 
+FOnConcertClientSyncSessionStartupOrShutdown& FConcertSyncClient::OnSyncSessionStartup()
+{
+	return OnSyncSessionStartupDelegate;
+}
+
+FOnConcertClientSyncSessionStartupOrShutdown& FConcertSyncClient::OnSyncSessionShutdown()
+{
+	return OnSyncSessionShutdownDelegate;
+}
+
 void FConcertSyncClient::PersistAllSessionChanges()
 {
 #if WITH_EDITOR
@@ -147,6 +157,7 @@ void FConcertSyncClient::RegisterConcertSyncHandlers(TSharedRef<IConcertClientSe
 	LiveSession = MakeShared<FConcertSyncClientLiveSession>(InSession, SessionFlags);
 	if (LiveSession->IsValidSession())
 	{
+		OnSyncSessionStartupDelegate.Broadcast(this);
 		CreateWorkspace(LiveSession.ToSharedRef());
 #if WITH_EDITOR
 		PresenceManager.Reset();
@@ -173,6 +184,7 @@ void FConcertSyncClient::UnregisterConcertSyncHandlers(TSharedRef<IConcertClient
 	PresenceManager.Reset();
 #endif
 	DestroyWorkspace();
+	OnSyncSessionShutdownDelegate.Broadcast(this);
 	LiveSession.Reset();
 }
 
diff --git a/Engine/Plugins/Developer/Concert/ConcertSync/ConcertSyncClient/Source/ConcertSyncClient/Private/ConcertSyncClient.h b/Engine/Plugins/Developer/Concert/ConcertSync/ConcertSyncClient/Source/ConcertSyncClient/Private/ConcertSyncClient.h
--- a/Engine/Plugins/Developer/Concert/ConcertSync/ConcertSyncClient/Source/ConcertSyncClient/Private/ConcertSyncClient.h
+++ b/Engine/Plugins/Developer/Concert/ConcertSync/ConcertSyncClient/Source/ConcertSyncClient/Private/ConcertSyncClient.h
@@ -33,6 +33,9 @@ public:
 	virtual FOnConcertClientWorkspaceStartupOrShutdown& OnWorkspaceStartup() override;
 	virtual FOnConcertClientWorkspaceStartupOrShutdown& OnWorkspaceShutdown() override;
 
+	virtual FOnConcertClientSyncSessionStartupOrShutdown& OnSyncSessionStartup() override;
+	virtual FOnConcertClientSyncSessionStartupOrShutdown& OnSyncSessionShutdown() override;
+
 	virtual void PersistAllSessionChanges() override;
 	virtual void GetSessionClientActions(const FConcertSessionClientInfo& InClientInfo, TArray<FConcertActionDefinition>& OutActions) const override;
 
@@ -67,6 +70,12 @@ private:
 	/** Delegate called on every workspace shutdown. */
 	FOnConcertClientWorkspaceStartupOrShutdown OnWorkspaceShutdownDelegate;
 
+	/** Delegate called just before the workspace get created. */
+	FOnConcertClientSyncSessionStartupOrShutdown OnSyncSessionStartupDelegate;
+	
+	/** Delegate called just after the workspace was deleted. */
+	FOnConcertClientSyncSessionStartupOrShutdown OnSyncSessionShutdownDelegate;
+
 #if WITH_EDITOR
 	/** Presence manager for the current session. */
 	TSharedPtr<FConcertClientPresenceManager> PresenceManager;
diff --git a/Engine/Plugins/Developer/Concert/ConcertSync/ConcertSyncClient/Source/ConcertSyncClient/Public/IConcertClientTransactionBridge.h b/Engine/Plugins/Developer/Concert/ConcertSync/ConcertSyncClient/Source/ConcertSyncClient/Public/IConcertClientTransactionBridge.h
--- a/Engine/Plugins/Developer/Concert/ConcertSync/ConcertSyncClient/Source/ConcertSyncClient/Public/IConcertClientTransactionBridge.h
+++ b/Engine/Plugins/Developer/Concert/ConcertSync/ConcertSyncClient/Source/ConcertSyncClient/Public/IConcertClientTransactionBridge.h
@@ -30,6 +30,7 @@ struct FConcertClientLocalTransactionCommonData
 	TArray<FName> ModifiedPackages;
 	TArray<FConcertObjectId> ExcludedObjectUpdates;
 	bool bIsExcluded = false;
+	bool bReplayable = true;
 };
 
 /**
diff --git a/Engine/Plugins/Developer/Concert/ConcertSync/ConcertSyncClient/Source/ConcertSyncClient/Public/IConcertClientWorkspace.h b/Engine/Plugins/Developer/Concert/ConcertSync/ConcertSyncClient/Source/ConcertSyncClient/Public/IConcertClientWorkspace.h
--- a/Engine/Plugins/Developer/Concert/ConcertSync/ConcertSyncClient/Source/ConcertSyncClient/Public/IConcertClientWorkspace.h
+++ b/Engine/Plugins/Developer/Concert/ConcertSync/ConcertSyncClient/Source/ConcertSyncClient/Public/IConcertClientWorkspace.h
@@ -158,4 +158,13 @@ public:
 	 * @param[in] OtherClientsWithModifMaxFetchNum The maximum number of client info to store in OutOtherClientsWithModifInfo if the latter is not null.
 	 */
 	virtual bool IsAssetModifiedByOtherClients(const FName& AssetName, int32* OutOtherClientsWithModifNum = nullptr, TArray<FConcertClientInfo>* OutOtherClientsWithModifInfo = nullptr, int32 OtherClientsWithModifMaxFetchNum = 0) const = 0;
+
+	/**
+	 * Controls whether the events emitted to the server through this workspace are marked as 'replayable'. By default, all events are marked as replayable. When
+	 * the workspace 'replayable' flag is false, the events emitted are recorded by the server, but are not replayed on the client(s). Non-replaylable events
+	 * are used for inspection rather than propagating/applying changes.
+	 * @note This was implemented to prevent multi-users transactions from being reapplied by disaster recovery in case of a crash during as multi-user session.
+	 * @param bReplayable True will mark all further events sent as replaylable, false as non-replayable.
+	 */
+	virtual void SetEmittedEventsAsReplayable(bool bReplayable) = 0;
 };
diff --git a/Engine/Plugins/Developer/Concert/ConcertSync/ConcertSyncClient/Source/ConcertSyncClient/Public/IConcertSyncClient.h b/Engine/Plugins/Developer/Concert/ConcertSync/ConcertSyncClient/Source/ConcertSyncClient/Public/IConcertSyncClient.h
--- a/Engine/Plugins/Developer/Concert/ConcertSync/ConcertSyncClient/Source/ConcertSyncClient/Public/IConcertSyncClient.h
+++ b/Engine/Plugins/Developer/Concert/ConcertSync/ConcertSyncClient/Source/ConcertSyncClient/Public/IConcertSyncClient.h
@@ -13,8 +13,10 @@ class IConcertClientPresenceManager;
 class IConcertClientSequencerManager;
 struct FConcertSessionClientInfo;
 
+class IConcertSyncClient;
 
 DECLARE_MULTICAST_DELEGATE_OneParam(FOnConcertClientWorkspaceStartupOrShutdown, const TSharedPtr<IConcertClientWorkspace>& /** InClientWorkspace */ );
+DECLARE_MULTICAST_DELEGATE_OneParam(FOnConcertClientSyncSessionStartupOrShutdown, const IConcertSyncClient* /* InSyncClient */ );
 
 /** Defines the supported editor play modes. */
 enum class EEditorPlayMode : uint8
@@ -67,6 +69,12 @@ public:
 	/** Get the delegate called on every workspace shutdown. */
 	virtual FOnConcertClientWorkspaceStartupOrShutdown& OnWorkspaceShutdown() = 0;
 
+	/** Get the delegate called just before the workspace gets created. */
+	virtual FOnConcertClientSyncSessionStartupOrShutdown& OnSyncSessionStartup() = 0;
+
+	/** Get the delegate called just after the workspace gets deleted. */
+	virtual FOnConcertClientSyncSessionStartupOrShutdown& OnSyncSessionShutdown() = 0;
+
 	/** Persist all session changes and prepare the files for source control submission. */
 	virtual void PersistAllSessionChanges() = 0;
 
diff --git a/Engine/Plugins/Developer/Concert/ConcertSync/ConcertSyncCore/Source/ConcertSyncCore/Public/ConcertTransactionEvents.h b/Engine/Plugins/Developer/Concert/ConcertSync/ConcertSyncCore/Source/ConcertSyncCore/Public/ConcertTransactionEvents.h
--- a/Engine/Plugins/Developer/Concert/ConcertSync/ConcertSyncCore/Source/ConcertSyncCore/Public/ConcertTransactionEvents.h
+++ b/Engine/Plugins/Developer/Concert/ConcertSync/ConcertSyncCore/Source/ConcertSyncCore/Public/ConcertTransactionEvents.h
@@ -129,6 +129,10 @@ struct FConcertTransactionEventBase
 
 	UPROPERTY()
 	TArray<FConcertExportedObject> ExportedObjects;
+
+	/** True to set the transaction event as replayable by clients. False would record it for logging/inspection purpose but will prevent clients from replaying it. */
+	UPROPERTY()
+	bool bReplayable = true;
 };
 
 USTRUCT()
diff --git a/Engine/Plugins/Developer/Concert/ConcertSync/ConcertSyncCore/Source/ConcertSyncCore/Public/ConcertWorkspaceMessages.h b/Engine/Plugins/Developer/Concert/ConcertSync/ConcertSyncCore/Source/ConcertSyncCore/Public/ConcertWorkspaceMessages.h
--- a/Engine/Plugins/Developer/Concert/ConcertSync/ConcertSyncCore/Source/ConcertSyncCore/Public/ConcertWorkspaceMessages.h
+++ b/Engine/Plugins/Developer/Concert/ConcertSync/ConcertSyncCore/Source/ConcertSyncCore/Public/ConcertWorkspaceMessages.h
@@ -73,6 +73,10 @@ struct FConcertPackageUpdateEvent
 
 	UPROPERTY()
 	FConcertPackage Package;
+
+	/** True to set the package event as replayable by clients. False would record it for logging/inspection purpose but will prevent clients from replaying it. */
+	UPROPERTY()
+	bool bReplayable = true;
 };
 
 USTRUCT()
diff --git a/Engine/Plugins/Developer/Concert/ConcertSync/ConcertSyncServer/Source/ConcertSyncServer/Private/ConcertServerWorkspace.cpp b/Engine/Plugins/Developer/Concert/ConcertSync/ConcertSyncServer/Source/ConcertSyncServer/Private/ConcertServerWorkspace.cpp
--- a/Engine/Plugins/Developer/Concert/ConcertSync/ConcertSyncServer/Source/ConcertSyncServer/Private/ConcertServerWorkspace.cpp
+++ b/Engine/Plugins/Developer/Concert/ConcertSync/ConcertSyncServer/Source/ConcertSyncServer/Private/ConcertServerWorkspace.cpp
@@ -260,7 +260,7 @@ void FConcertServerWorkspace::HandlePackageUpdateEvent(const FConcertSessionCont
 	{
 		return;
 	}
-	
+
 	// Consider acquiring lock on asset saving an explicit lock
 	const bool bLockOwned = LockWorkspaceResource(Event.Package.Info.PackageName, Context.SourceEndpointId, EConcertLockFlags::Temporary);
 	if (bLockOwned)
@@ -280,6 +280,7 @@ void FConcertServerWorkspace::HandlePackageUpdateEvent(const FConcertSessionCont
 				}
 			}
 			PackageActivity.EventSummary.SetTypedPayload(FConcertSyncPackageActivitySummary::CreateSummaryForEvent(PackageActivity.EventData));
+			PackageActivity.bIgnored = !Event.bReplayable;
 			AddPackageActivity(PackageActivity);
 		}
 
@@ -322,6 +323,7 @@ void FConcertServerWorkspace::HandleTransactionFinalizedEvent(const FConcertSess
 			TransactionActivity.EndpointId = InEventContext.SourceEndpointId;
 			TransactionActivity.EventData.Transaction = InEvent;
 			TransactionActivity.EventSummary.SetTypedPayload(FConcertSyncTransactionActivitySummary::CreateSummaryForEvent(TransactionActivity.EventData));
+			TransactionActivity.bIgnored = !InEvent.bReplayable;
 			AddTransactionActivity(TransactionActivity);
 		}
 
@@ -340,7 +342,7 @@ void FConcertServerWorkspace::HandleTransactionFinalizedEvent(const FConcertSess
 
 void FConcertServerWorkspace::HandleTransactionSnapshotEvent(const FConcertSessionContext& InEventContext, const FConcertTransactionSnapshotEvent& InEvent)
 {
-	if (!EnumHasAnyFlags(LiveSession->GetSessionFlags(), EConcertSyncSessionFlags::EnableTransactions))
+	if (!EnumHasAnyFlags(LiveSession->GetSessionFlags(), EConcertSyncSessionFlags::EnableTransactions) || !InEvent.bReplayable)
 	{
 		return;
 	}
diff --git a/Engine/Plugins/Developer/Concert/DisasterRecoveryClient/Source/DisasterRecoveryClient/Private/DisasterRecoveryClientModule.cpp b/Engine/Plugins/Developer/Concert/DisasterRecoveryClient/Source/DisasterRecoveryClient/Private/DisasterRecoveryClientModule.cpp
--- a/Engine/Plugins/Developer/Concert/DisasterRecoveryClient/Source/DisasterRecoveryClient/Private/DisasterRecoveryClientModule.cpp
+++ b/Engine/Plugins/Developer/Concert/DisasterRecoveryClient/Source/DisasterRecoveryClient/Private/DisasterRecoveryClientModule.cpp
@@ -6,6 +6,7 @@
 #include "IConcertSyncClientModule.h"
 #include "IConcertModule.h"
 #include "IConcertClient.h"
+#include "IConcertClientWorkspace.h"
 #include "IConcertSession.h"
 #include "IConcertSyncClient.h"
 #include "ConcertFrontendStyle.h"
@@ -50,6 +51,14 @@ public:
 		// Wait for init to finish before starting the Disaster Recovery service
 		FCoreDelegates::OnFEngineLoopInitComplete.AddRaw(this, &FDisasterRecoveryClientModule::OnEngineInitComplete);
 
+		// Hook to listen when a new session is created.
+		IConcertSyncClientModule::Get().OnClientCreated().AddRaw(this, &FDisasterRecoveryClientModule::HandleConcertSyncClientCreated);
+		for (TSharedRef<IConcertSyncClient> Client : IConcertSyncClientModule::Get().GetClients())
+		{
+			Client->OnSyncSessionStartup().AddRaw(this, &FDisasterRecoveryClientModule::HandleSyncSessionStartup);
+			Client->OnSyncSessionShutdown().AddRaw(this, &FDisasterRecoveryClientModule::HandleSyncSessionShutdown);
+		}
+
 		// Initialize Style
 		FConcertFrontendStyle::Initialize();
 
@@ -65,6 +74,17 @@ public:
 		FCoreDelegates::OnPreExit.RemoveAll(this);
 		HandleAppPreExit();
 
+		// Unhook this module callback from other clients.
+		if (IConcertSyncClientModule::IsAvailable())
+		{
+			IConcertSyncClientModule::Get().OnClientCreated().RemoveAll(this);
+			for (TSharedRef<IConcertSyncClient> Client : IConcertSyncClientModule::Get().GetClients())
+			{
+				Client->OnSyncSessionStartup().RemoveAll(this);
+				Client->OnSyncSessionShutdown().RemoveAll(this);
+			}
+		}
+
 		// Unregister the Disaster Recovery Settings panel.
 		UnregisterSettings();
 	}
@@ -138,6 +158,27 @@ private:
 		StopDisasterRecoveryService();
 	}
 
+	void HandleConcertSyncClientCreated(TSharedRef<IConcertSyncClient> Client)
+	{
+		if (Client->GetConcertClient()->GetRole() != Role) // Exclude disaster recovery own session connection changes.
+		{
+			Client->OnSyncSessionStartup().AddRaw(this, &FDisasterRecoveryClientModule::HandleSyncSessionStartup);
+			Client->OnSyncSessionShutdown().AddRaw(this, &FDisasterRecoveryClientModule::HandleSyncSessionShutdown);
+		}
+	}
+
+	void HandleSyncSessionStartup(const IConcertSyncClient* SyncClient)
+	{
+		check(DisasterRecoveryClient.Get() != SyncClient)
+		SetDisasterRecoveryEventsAsReplayable(IsCompatibleWithOtherConcertSessions(SyncClient, /*SyncClientShuttingDownSession*/nullptr));
+	}
+
+	void HandleSyncSessionShutdown(const IConcertSyncClient* SyncClient)
+	{
+		check(DisasterRecoveryClient.Get() != SyncClient);
+		SetDisasterRecoveryEventsAsReplayable(IsCompatibleWithOtherConcertSessions(/*SyncClientStartingSession*/nullptr, SyncClient));
+	}
+
 	FGuid GetDisasterRecoverySessionId() const
 	{
 		if (DisasterRecoveryClient)
@@ -251,7 +292,7 @@ private:
 		// Normally, bAutoRestoreLastSession is set true here and overwritten to false when the app exits normally, but when running under the debugger, auto-restore is disabled as
 		// programmers kill applications (stop the debugger) and this should not count as a crash (unless you want to simulate crash this way - see below).
 		Session.bAutoRestoreLastSession = !FPlatformMisc::IsDebuggerPresent();
-		//Sessions.bAutoRestoreLastSession = true; // <- MUST BE COMMENTED BEFORE SUBMITTING: For debugging purpose only. Simulate a crash by stopping the debugger during a session.
+		//Session.bAutoRestoreLastSession = true; // <- MUST BE COMMENTED BEFORE SUBMITTING: For debugging purpose only. Simulate a crash by stopping the debugger during a session.
 
 		// Save the file.
 		SaveDisasterRecoverySessionInfo(RecoverySessionInfo);
@@ -377,6 +418,9 @@ private:
 		DisasterRecoveryClient->GetConcertClient()->OnSessionStartup().AddRaw(this, &FDisasterRecoveryClientModule::DisasterRecoverySessionCreated);
 		DisasterRecoveryClient->Startup(ClientConfig, EConcertSyncSessionFlags::Default_DisasterRecoverySession);
 
+		// Set all events captured by the disaster recovery service as 'replayable' unless another concert client (assumed Multi-User) has created an incompatible session.
+		SetDisasterRecoveryEventsAsReplayable(IsCompatibleWithOtherConcertSessions(/*SyncClientStartingSession*/nullptr, /*SyncClientShuttingDownSession*/nullptr));
+
 		// If something needs to be recovered from a crash.
 		if (RestoringSession)
 		{
@@ -428,6 +472,54 @@ private:
 		}
 	}
 
+	/** Returns true if disaster recovery Concert session can run concurrently with other Concert sessions (if any). */
+	bool IsCompatibleWithOtherConcertSessions(const IConcertSyncClient* SyncClientStartingSession, const IConcertSyncClient* SyncClientShuttingDownSession) const
+	{
+		// At the moment, we don't expect more than 2 clients. We don't have use cases for a third concurrent concert client.
+		checkf(IConcertSyncClientModule::Get().GetClients().Num() <= 2, TEXT("Expected 1 disaster recovery client + 1 multi-user client at max."));
+
+		// Scan all existing clients.
+		for (const TSharedRef<IConcertSyncClient>& SyncClient : IConcertSyncClientModule::Get().GetClients())
+		{
+			if (SyncClient == DisasterRecoveryClient || &SyncClient.Get() == SyncClientShuttingDownSession)
+			{
+				continue; // Compatible with itself or the sync client is shutting down its sync session, so it cannot interfere anymore.
+			}
+			else if (&SyncClient.Get() == SyncClientStartingSession)
+			{
+				if (!IsCompatibleWithConcertClient(&SyncClient.Get()))
+				{
+					return false; // The sync client starting a session will interfere with disaster recovery client.
+				}
+			}
+			else if (SyncClient->GetWorkspace() && !IsCompatibleWithConcertClient(&SyncClient.Get())) // A valid workspace means the client is joining, in or leaving a session.
+			{
+				return false; // That existing client is interfering with disaster recovery client.
+			}
+		}
+
+		return true; // No other sessions exist or it is compatible.
+	}
+
+	bool IsCompatibleWithConcertClient(const IConcertSyncClient* SyncClient) const
+	{
+		check(SyncClient != DisasterRecoveryClient.Get());
+		checkf(SyncClient->GetConcertClient()->GetRole() == TEXT("MultiUser"), TEXT("A new role was added, check if this role can run concurrently with disaster recovery."));
+
+		// Multi-User (MU) sessions are not compatible with disaster recovery (DR) session because MU events are performed in a transient sandbox that doesn't exist outside the MU session.
+		// If a crash occurs during a MU session, DR must not recover transactions applied to the transient sandbox. DR will will record the MU events, but for crash inspection purpose only.
+		return SyncClient->GetConcertClient()->GetRole() != TEXT("MultiUser");
+	}
+
+	/** Sets whether further Concert events (transaction/package) emitted by Disaster Recovery have the replayable flag on or off. */
+	void SetDisasterRecoveryEventsAsReplayable(bool bReplayable)
+	{
+		if (TSharedPtr<IConcertClientWorkspace> Workspace = DisasterRecoveryClient ? DisasterRecoveryClient->GetWorkspace() : TSharedPtr<IConcertClientWorkspace>())
+		{
+			Workspace->SetEmittedEventsAsReplayable(bReplayable);
+		}
+	}
+
 private:
 	/** This client role, a tag given to different types of concert client, i.e. DisasterRecovery for this one. */
 	FString Role;
diff --git a/Engine/Plugins/Developer/Concert/DisasterRecoveryClient/Source/DisasterRecoveryClient/Private/DisasterRecoveryFSM.cpp b/Engine/Plugins/Developer/Concert/DisasterRecoveryClient/Source/DisasterRecoveryClient/Private/DisasterRecoveryFSM.cpp
--- a/Engine/Plugins/Developer/Concert/DisasterRecoveryClient/Source/DisasterRecoveryClient/Private/DisasterRecoveryFSM.cpp
+++ b/Engine/Plugins/Developer/Concert/DisasterRecoveryClient/Source/DisasterRecoveryClient/Private/DisasterRecoveryFSM.cpp
@@ -366,6 +366,7 @@ void FDisasterRecoveryFSM::RestoreAndJoinSession()
 	RestoreArgs.SessionId = RecoverySessionId;
 	RestoreArgs.SessionFilter.ActivityIdUpperBound = SelectedRecoveryActivity->Activity.ActivityId;
 	RestoreArgs.SessionFilter.bOnlyLiveData = bLiveDataOnly; // TODO: I'm not sure if this should be true or false.
+	RestoreArgs.SessionFilter.bIncludeIgnoredActivities = false; // Don't restore ignored activities. (Like MultiUser events recorded in a DisasterRecovery session for inspection purpose)
 	RestoreArgs.SessionName = ClientConfig->DefaultSessionName;
 
 	// Restore the session on the server.
