Based the decision to launch the disaster recovery standalone server on FGenericCrashContext::IsOutOfProcessCrashReporter(). If out of process crash reporter is true, then it is assumed to host disaster recovery as well.

diff --git a/Engine/Plugins/Developer/Concert/DisasterRecoveryClient/Source/DisasterRecoveryClient/Private/DisasterRecoveryClientModule.cpp b/Engine/Plugins/Developer/Concert/DisasterRecoveryClient/Source/DisasterRecoveryClient/Private/DisasterRecoveryClientModule.cpp
--- a/Engine/Plugins/Developer/Concert/DisasterRecoveryClient/Source/DisasterRecoveryClient/Private/DisasterRecoveryClientModule.cpp
+++ b/Engine/Plugins/Developer/Concert/DisasterRecoveryClient/Source/DisasterRecoveryClient/Private/DisasterRecoveryClientModule.cpp
@@ -95,15 +95,6 @@ public:
 	}
 
 private:
-	constexpr bool ShouldLaunchDisasterRecoveryBackendService() const
-	{
-#if PLATFORM_WINDOWS // On windows, the service is launched by the crash analytics process. (WindowsPlatformCrashContext.cpp)
-		return false;
-#else
-		return true;
-#endif
-	}
-
 	void RegisterSettings()
 	{
 		if (ISettingsModule* SettingsModule = FModuleManager::GetModulePtr<ISettingsModule>("Settings"))
@@ -394,7 +385,8 @@ private:
 		const FString DisasterRecoveryServerName = RecoveryService::GetRecoveryServerName();
 		const FString DisasterRecoverySessionName = FString::Printf(TEXT("%s_%s_%s"), *DisasterRecoveryServerName, FApp::GetProjectName(), *FDateTime::Now().ToString());
 
-		if (ShouldLaunchDisasterRecoveryBackendService() && !SpawnDisasterRecoveryServer(DisasterRecoveryServerName))
+		// If crash reporter is running out of process, it also hosts disaster recovery server as the '-ConcertServer' param is set when spawning CrashReporterClient. No need to start the UnrealDisasterRecoveryService executable.
+		if (!FGenericCrashContext::IsOutOfProcessCrashReporter() && !SpawnDisasterRecoveryServer(DisasterRecoveryServerName))
 		{
 			return false; // Failed to spawn the service.
 		}
