- Fixed disaster recovery remote endpoint timeout set to zero, preventing it from re-registering with MessageBus when an error occurred (like the socket disconnected).

  - Fixed the crash reporter app to display the UI (asking the user to send the bug report) before shutting down the recovery service.

- Renamed the field FDisasterRecoveryInfo::Version into FDisasterRecoveryInfo::Revision because revision is more accurate for the field.

Edigrated 11250824 from Dev-VirtualProduction.

diff --git a/Engine/Plugins/Developer/Concert/DisasterRecoveryClient/Source/DisasterRecoveryClient/Private/DisasterRecoveryClientModule.cpp b/Engine/Plugins/Developer/Concert/DisasterRecoveryClient/Source/DisasterRecoveryClient/Private/DisasterRecoveryClientModule.cpp
--- a/Engine/Plugins/Developer/Concert/DisasterRecoveryClient/Source/DisasterRecoveryClient/Private/DisasterRecoveryClientModule.cpp
+++ b/Engine/Plugins/Developer/Concert/DisasterRecoveryClient/Source/DisasterRecoveryClient/Private/DisasterRecoveryClientModule.cpp
@@ -99,8 +99,8 @@ public:
 		ClientConfig->DefaultServerURL = DisasterRecoveryServerName;
 		ClientConfig->DefaultSessionName = DisasterRecoverySessionName;
 		ClientConfig->DefaultSaveSessionAs = DisasterRecoverySessionName;
-		//ClientConfig->ClientSettings.DiscoveryTimeoutSeconds = 0;
-		ClientConfig->EndpointSettings.RemoteEndpointTimeoutSeconds = 0;
+		//ClientConfig->ClientSettings.DiscoveryTimeoutSeconds = 0; // Setting this to zero prevents showing 'Server {serverName} lost." but will not detect when a server dies, growing the list of known server indefinitely.
+		ClientConfig->EndpointSettings.RemoteEndpointTimeoutSeconds = 4; // Non-zero -> Ensure the Endpoints are re-registered with Message Bus if Message Bus restarts or some errors trigger a 'auto-repair'.
 		ClientConfig->ClientSettings.ClientAuthenticationKey = DisasterRecoveryServerName; // The server adds its own server name to the list of authorized client keys, use that key to authorize this client on the server.
 
 		// Create and start the client.
diff --git a/Engine/Plugins/Developer/Concert/DisasterRecoveryClient/Source/DisasterRecoveryClient/Private/DisasterRecoverySessionInfo.h b/Engine/Plugins/Developer/Concert/DisasterRecoveryClient/Source/DisasterRecoveryClient/Private/DisasterRecoverySessionInfo.h
--- a/Engine/Plugins/Developer/Concert/DisasterRecoveryClient/Source/DisasterRecoveryClient/Private/DisasterRecoverySessionInfo.h
+++ b/Engine/Plugins/Developer/Concert/DisasterRecoveryClient/Source/DisasterRecoveryClient/Private/DisasterRecoverySessionInfo.h
@@ -103,9 +103,9 @@ struct FDisasterRecoveryInfo
 {
 	GENERATED_BODY()
 
-	/** The version of the information. Updated everytime the recovery info file is written. */
+	/** The revision number of the information. Updated everytime the recovery info file is written. */
 	UPROPERTY()
-	uint32 Version = 0;
+	uint32 Revision = 0;
 
 	/** The list of running/crashing/crashed sessions. */
 	UPROPERTY()
diff --git a/Engine/Plugins/Developer/Concert/DisasterRecoveryClient/Source/DisasterRecoveryClient/Private/DisasterRecoverySessionManager.cpp b/Engine/Plugins/Developer/Concert/DisasterRecoveryClient/Source/DisasterRecoveryClient/Private/DisasterRecoverySessionManager.cpp
--- a/Engine/Plugins/Developer/Concert/DisasterRecoveryClient/Source/DisasterRecoveryClient/Private/DisasterRecoverySessionManager.cpp
+++ b/Engine/Plugins/Developer/Concert/DisasterRecoveryClient/Source/DisasterRecoveryClient/Private/DisasterRecoverySessionManager.cpp
@@ -94,9 +94,9 @@ bool SaveDisasterRecoveryInfo(const FString& Pathname, const FDisasterRecoveryIn
 class ScopedRecoveryInfoExclusiveUpdater
 {
 public:
-	ScopedRecoveryInfoExclusiveUpdater(const FString& InPathname, uint32& OutWrittenVersion)
+	ScopedRecoveryInfoExclusiveUpdater(const FString& InPathname, uint32& OutWrittenRevision)
 		: Pathname(InPathname)
-		, WrittenVersion(OutWrittenVersion)
+		, WrittenRevision(OutWrittenRevision)
 		, SystemWideMutex(GetSystemMutexName()) // Get machine-wide exclusive access to the file during the scope.
 	{
 		LoadDisasterRecoveryInfo(Pathname, RecoveryInfo);
@@ -106,9 +106,9 @@ public:
 	{
 		if (bAutoSave)
 		{
-			RecoveryInfo.Version += 1; // The file is going to be written, increase the version.
+			RecoveryInfo.Revision += 1; // The file is going to be written, increase the revision number.
 			SaveDisasterRecoveryInfo(Pathname, RecoveryInfo);
-			WrittenVersion = RecoveryInfo.Version;
+			WrittenRevision = RecoveryInfo.Revision;
 		}
 	}
 
@@ -117,7 +117,7 @@ public:
 
 private:
 	const FString& Pathname;
-	uint32& WrittenVersion;
+	uint32& WrittenRevision;
 	FSystemWideCriticalSection SystemWideMutex;
 	FDisasterRecoveryInfo RecoveryInfo;
 	bool bAutoSave = true; // By default, automatically save when going out of scope.
@@ -131,7 +131,7 @@ uint32 GetDisasterRecoveryInfoVersion(const FString& InPathname)
 
 	FDisasterRecoveryInfo RecoveryInfo;
 	LoadDisasterRecoveryInfo(InPathname, RecoveryInfo);
-	return RecoveryInfo.Version;
+	return RecoveryInfo.Revision;
 }
 
 /** Migrate the recovery info file from 4.24 to the new format in 4.25. */
@@ -289,18 +289,18 @@ FDisasterRecoverySessionManager::~FDisasterRecoverySessionManager()
 	SyncClient->GetConcertClient()->OnSessionShutdown().RemoveAll(this);
 
 	// Remove this recovery clients from the list of active clients.
-	DisasterRecoveryUtil::ScopedRecoveryInfoExclusiveUpdater ScopedRecoveryUpdater(RecoveryInfoPathname, SessionsCacheVersion);
+	DisasterRecoveryUtil::ScopedRecoveryInfoExclusiveUpdater ScopedRecoveryUpdater(RecoveryInfoPathname, SessionsCacheRevision);
 	ScopedRecoveryUpdater.Info().ClientProcessIds.RemoveAll([](const int32 ProcessId) { return ProcessId == FPlatformProcess::GetCurrentProcessId(); });
 }
 
 TArray<FGuid> FDisasterRecoverySessionManager::InitAndRotateSessions()
 {
-	DisasterRecoveryUtil::ScopedRecoveryInfoExclusiveUpdater ScopedRecoveryUpdater(RecoveryInfoPathname, SessionsCacheVersion);
+	DisasterRecoveryUtil::ScopedRecoveryInfoExclusiveUpdater ScopedRecoveryUpdater(RecoveryInfoPathname, SessionsCacheRevision);
 	FDisasterRecoveryInfo& RecoveryInfo = ScopedRecoveryUpdater.Info();
 
 	// If the user is upgrading from 4.24 to 4.25 the name and the format of the recovery info file changed. Try migrating to the new format.
 	FString OldSessionInfo = GetDisasterRecoveryInfoPath() / TEXT("Sessions.json");
-	if (IFileManager::Get().FileExists(*OldSessionInfo) && RecoveryInfo.Version == 0)
+	if (IFileManager::Get().FileExists(*OldSessionInfo) && RecoveryInfo.Revision == 0)
 	{
 		DisasterRecoveryUtil::MigrateSessionInfoFrom_4_24(OldSessionInfo, RecoveryInfo);
 		IFileManager::Get().Delete(*OldSessionInfo);
@@ -416,11 +416,11 @@ void FDisasterRecoverySessionManager::Refresh()
 	};
 
 	bool bSessionUpdated = false;
-	uint32 VersionBefore;
+	uint32 RevisionBefore;
 	{
-		DisasterRecoveryUtil::ScopedRecoveryInfoExclusiveUpdater ScopedRecoveryUpdater(RecoveryInfoPathname, SessionsCacheVersion);
+		DisasterRecoveryUtil::ScopedRecoveryInfoExclusiveUpdater ScopedRecoveryUpdater(RecoveryInfoPathname, SessionsCacheRevision);
 		FDisasterRecoveryInfo& RecoveryInfo = ScopedRecoveryUpdater.Info();
-		VersionBefore = RecoveryInfo.Version;
+		RevisionBefore = RecoveryInfo.Revision;
 
 		bSessionUpdated |= UpdateCrashedSessionInfoFn(RecoveryInfo.ActiveSessions);
 		bSessionUpdated |= UpdateCrashedSessionInfoFn(RecoveryInfo.RecentSessions);
@@ -429,7 +429,7 @@ void FDisasterRecoverySessionManager::Refresh()
 		ScopedRecoveryUpdater.SetAutoSave(bSessionUpdated);
 	}
 
-	if (bSessionUpdated || VersionBefore != SessionsCacheVersion) // The file was changed either by the ScopedRecoveryInfoExclusiveUpdater above or by another instance.
+	if (bSessionUpdated || RevisionBefore != SessionsCacheRevision) // The file was changed either by the ScopedRecoveryInfoExclusiveUpdater above or by another instance.
 	{
 		UpdateSessionsCache();
 	}
@@ -444,7 +444,7 @@ void FDisasterRecoverySessionManager::UpdateSessionsCache()
 
 		// Load the session info file (if it exist)
 		DisasterRecoveryUtil::LoadDisasterRecoveryInfo(RecoveryInfoPathname, RecoveryInfo);
-		SessionsCacheVersion = RecoveryInfo.Version;
+		SessionsCacheRevision = RecoveryInfo.Revision;
 	}
 
 	// Remove from the cache the sessions that were removed from the recovery info file (possibly by another instance).
@@ -947,7 +947,7 @@ bool FDisasterRecoverySessionManager::HasInProgressSession() const
 void FDisasterRecoverySessionManager::OnRecoverySessionStartup(TSharedRef<IConcertClientSession> InSession)
 {
 	{
-		DisasterRecoveryUtil::ScopedRecoveryInfoExclusiveUpdater ScopedRecoveryUpdater(RecoveryInfoPathname, SessionsCacheVersion);
+		DisasterRecoveryUtil::ScopedRecoveryInfoExclusiveUpdater ScopedRecoveryUpdater(RecoveryInfoPathname, SessionsCacheRevision);
 		FDisasterRecoveryInfo& RecoveryInfo = ScopedRecoveryUpdater.Info();
 
 		// The user had the chance to 'review' the 'pending' candidates, move them to the 'recent' list to not review them again.
@@ -990,7 +990,7 @@ void FDisasterRecoverySessionManager::OnRecoverySessionStartup(TSharedRef<IConce
 void FDisasterRecoverySessionManager::OnRecoverySessionShutdown(TSharedRef<IConcertClientSession> InSession)
 {
 	{
-		DisasterRecoveryUtil::ScopedRecoveryInfoExclusiveUpdater ScopedRecoveryUpdater(RecoveryInfoPathname, SessionsCacheVersion);
+		DisasterRecoveryUtil::ScopedRecoveryInfoExclusiveUpdater ScopedRecoveryUpdater(RecoveryInfoPathname, SessionsCacheRevision);
 		FDisasterRecoveryInfo& RecoveryInfo = ScopedRecoveryUpdater.Info();
 
 		// The session ended up normally, move the current session from the 'active' list to the 'recent' list.
@@ -1026,7 +1026,7 @@ void FDisasterRecoverySessionManager::OnRecoverySessionShutdown(TSharedRef<IConc
 void FDisasterRecoverySessionManager::OnSessionNotFoundInternal(const FGuid& RepositoryId)
 {
 	{
-		DisasterRecoveryUtil::ScopedRecoveryInfoExclusiveUpdater ScopedRecoveryUpdater(RecoveryInfoPathname, SessionsCacheVersion);
+		DisasterRecoveryUtil::ScopedRecoveryInfoExclusiveUpdater ScopedRecoveryUpdater(RecoveryInfoPathname, SessionsCacheRevision);
 		FDisasterRecoveryInfo& RecoveryInfo = ScopedRecoveryUpdater.Info();
 
 		// Mark the session repository to be discarded.
@@ -1046,7 +1046,7 @@ void FDisasterRecoverySessionManager::OnSessionNotFoundInternal(const FGuid& Rep
 void FDisasterRecoverySessionManager::OnSessionRepositoriesDiscardedInternal(const TArray<FGuid>& DiscardedRepositoryIds)
 {
 	{
-		DisasterRecoveryUtil::ScopedRecoveryInfoExclusiveUpdater ScopedRecoveryUpdater(RecoveryInfoPathname, SessionsCacheVersion);
+		DisasterRecoveryUtil::ScopedRecoveryInfoExclusiveUpdater ScopedRecoveryUpdater(RecoveryInfoPathname, SessionsCacheRevision);
 		FDisasterRecoveryInfo& RecoveryInfo = ScopedRecoveryUpdater.Info();
 
 		// If the repository was discarded, the session cannot be restored or visualized anymore, clear it from the session database.
@@ -1074,7 +1074,7 @@ void FDisasterRecoverySessionManager::OnSessionRepositoryMountedInternal(const F
 	int32 ThisProcessId = FPlatformProcess::GetCurrentProcessId();
 
 	{
-		DisasterRecoveryUtil::ScopedRecoveryInfoExclusiveUpdater ScopedRecoveryUpdater(RecoveryInfoPathname, SessionsCacheVersion);
+		DisasterRecoveryUtil::ScopedRecoveryInfoExclusiveUpdater ScopedRecoveryUpdater(RecoveryInfoPathname, SessionsCacheRevision);
 		FDisasterRecoveryInfo& RecoveryInfo = ScopedRecoveryUpdater.Info();
 
 		if (FDisasterRecoverySession* Active = RecoveryInfo.ActiveSessions.FindByPredicate([&RepositoryId](const FDisasterRecoverySession& Candidate) { return Candidate.RepositoryId == RepositoryId; }))
@@ -1112,7 +1112,7 @@ void FDisasterRecoverySessionManager::OnSessionRepositoryMountedInternal(const F
 void FDisasterRecoverySessionManager::OnSessionImportedInternal(const TSharedRef<FDisasterRecoverySession>& ImportedSession)
 {
 	{
-		DisasterRecoveryUtil::ScopedRecoveryInfoExclusiveUpdater ScopedRecoveryUpdater(RecoveryInfoPathname, SessionsCacheVersion);
+		DisasterRecoveryUtil::ScopedRecoveryInfoExclusiveUpdater ScopedRecoveryUpdater(RecoveryInfoPathname, SessionsCacheRevision);
 		ScopedRecoveryUpdater.Info().ImportedSessions.Insert(*ImportedSession, 0); // Most recently import are at front.
 	}
 
@@ -1138,7 +1138,7 @@ void FDisasterRecoverySessionManager::OnSessionRepositoryFilesModified(const TAr
 	{
 		if (FileChangeData.Action == FFileChangeData::EFileChangeAction::FCA_Modified && FileChangeData.Filename.EndsWith(GetDisasterRecoveryInfoFilename()))
 		{
-			if (DisasterRecoveryUtil::GetDisasterRecoveryInfoVersion(RecoveryInfoPathname) != SessionsCacheVersion) // Skip refresh is it wasn't modified by another process.
+			if (DisasterRecoveryUtil::GetDisasterRecoveryInfoVersion(RecoveryInfoPathname) != SessionsCacheRevision) // Skip refresh is it wasn't modified by another process.
 			{
 				Refresh();
 			}
diff --git a/Engine/Plugins/Developer/Concert/DisasterRecoveryClient/Source/DisasterRecoveryClient/Private/DisasterRecoverySessionManager.h b/Engine/Plugins/Developer/Concert/DisasterRecoveryClient/Source/DisasterRecoveryClient/Private/DisasterRecoverySessionManager.h
--- a/Engine/Plugins/Developer/Concert/DisasterRecoveryClient/Source/DisasterRecoveryClient/Private/DisasterRecoverySessionManager.h
+++ b/Engine/Plugins/Developer/Concert/DisasterRecoveryClient/Source/DisasterRecoveryClient/Private/DisasterRecoverySessionManager.h
@@ -179,6 +179,6 @@ private:
 	FString RecoveryInfoPathname;
 
 	TArray<int32> CrashingProcessIds;
-	uint32 SessionsCacheVersion = 0;
+	uint32 SessionsCacheRevision = 0;
 	bool bRefreshing = false;
 };
diff --git a/Engine/Source/Programs/CrashReportClient/Private/CrashReportClientApp.cpp b/Engine/Source/Programs/CrashReportClient/Private/CrashReportClientApp.cpp
--- a/Engine/Source/Programs/CrashReportClient/Private/CrashReportClientApp.cpp
+++ b/Engine/Source/Programs/CrashReportClient/Private/CrashReportClientApp.cpp
@@ -715,7 +715,7 @@ static bool WasAbnormalShutdown(const FEditorSessionSummarySender& SessionSummar
 	return false;
 }
 
-#endif 
+#endif
 
 void RunCrashReportClient(const TCHAR* CommandLine)
 {
@@ -830,18 +830,18 @@ void RunCrashReportClient(const TCHAR* CommandLine)
 
 					// Build error report in memory.
 					FPlatformErrorReport ErrorReport = CollectErrorReport(RecoveryServicePtr.Get(), MonitorPid, CrashContext, MonitorWritePipe);
+					const bool bNoDialog = (CrashContext.UserSettings.bNoDialog || CrashContext.UserSettings.bImplicitSend) && CrashContext.UserSettings.bSendUnattendedBugReports;
+					const SubmitCrashReportResult Result = SendErrorReport(ErrorReport, bNoDialog, CrashContext.UserSettings.bImplicitSend);
 
 #if CRASH_REPORT_WITH_RECOVERY
+					// Shut down the recovery service AFTER the UI prompt (if any) because archiving a large session can take several minutes (must be on main thread - UObject stuff). This ensures
+					// the UI is shown in a timely manner, but delays the server shutdown. At next reboot, user may not be able to restore because the session will still be archiving.
 					if (RecoveryServicePtr && !FPrimaryCrashProperties::Get()->bIsEnsure)
 					{
-						// Shutdown the recovery service, releasing the recovery database file lock (not sharable) as soon as possible to let a new instance take it and offer the user to recover.
+						// Shutdown the recovery service. This will releases the recovery database file lock (not sharable) and let a new instance take it and offer the user to recover.
 						RecoveryServicePtr.Reset();
 					}
 #endif
-
-					const bool bNoDialog = (CrashContext.UserSettings.bNoDialog || CrashContext.UserSettings.bImplicitSend) && CrashContext.UserSettings.bSendUnattendedBugReports;
-					const SubmitCrashReportResult Result = SendErrorReport(ErrorReport, bNoDialog, CrashContext.UserSettings.bImplicitSend);
-
 					if (bReportCrashAnalyticInfo)
 					{
 						// If analytics is enabled make sure they are submitted now.
diff --git a/Engine/Source/Programs/CrashReportClient/Private/RecoveryService.cpp b/Engine/Source/Programs/CrashReportClient/Private/RecoveryService.cpp
--- a/Engine/Source/Programs/CrashReportClient/Private/RecoveryService.cpp
+++ b/Engine/Source/Programs/CrashReportClient/Private/RecoveryService.cpp
@@ -84,7 +84,7 @@ bool FRecoveryService::Startup()
 	// Setup the disaster recovery server configuration
 	UConcertServerConfig* ServerConfig = IConcertSyncServerModule::Get().ParseServerSettings(FCommandLine::Get());
 	ServerConfig->bAutoArchiveOnReboot = true; // If server crashed, was killed, etc, ensure the recovery session is archived (expected by recovery flow).
-	ServerConfig->EndpointSettings.RemoteEndpointTimeoutSeconds = 0;
+	ServerConfig->EndpointSettings.RemoteEndpointTimeoutSeconds = 4; // Non-zero -> Ensure the Endpoints are re-registered with Message Bus if UDP layer restarts or some errors trigger a 'auto-repair'.
 	ServerConfig->bMountDefaultSessionRepository = false; // Let the client mount its own repository to support concurrent recovery server and prevent them from concurrently accessing non-sharable database files.
 	ServerConfig->AuthorizedClientKeys.Add(ServerConfig->ServerName); // The disaster recovery client is configured to use the unique server name as key to identify itself.
 
diff --git a/Engine/Source/Programs/UnrealDisasterRecoveryService/Private/UnrealDisasterRecoveryService.cpp b/Engine/Source/Programs/UnrealDisasterRecoveryService/Private/UnrealDisasterRecoveryService.cpp
--- a/Engine/Source/Programs/UnrealDisasterRecoveryService/Private/UnrealDisasterRecoveryService.cpp
+++ b/Engine/Source/Programs/UnrealDisasterRecoveryService/Private/UnrealDisasterRecoveryService.cpp
@@ -40,7 +40,7 @@ INT32_MAIN_INT32_ARGC_TCHAR_ARGV()
 
 		UConcertServerConfig* ServerConfig = IConcertSyncServerModule::Get().ParseServerSettings(FCommandLine::Get());
 		ServerConfig->bAutoArchiveOnReboot = true; // If server crashed, was killed, etc, ensure the recovery session is archived (expected by recovery flow).
-		ServerConfig->EndpointSettings.RemoteEndpointTimeoutSeconds = 0;
+		ServerConfig->EndpointSettings.RemoteEndpointTimeoutSeconds = 4; // Non-zero -> Ensure the Endpoints are re-registered with Message Bus if the UDP layer restarts or some errors trigger a 'auto-repair'.
 		ServerConfig->bMountDefaultSessionRepository = false; // Let the client mount its own repository to support concurrent service and avoid them to access the same non-sharable database files.
 		ServerConfig->AuthorizedClientKeys.Add(ServerConfig->ServerName); // The disaster recovery client is configured to use the unique server name as key to identify itself.
 		return ServerConfig;
