- Added the ability to copy and restore a live session, preventing the need to archive it in first place, making the server exist fast (releasing the session lock very quickly) before showing the crash UI and before the next Editor instance could starts.

Details:

This bug could manifest if various ways. An issue causing this bug was fixed in 11252374. This bug can also be observed if the crash reporting process doesn't release its lock on the crashed session quickly. Archiving a session may takes several minutes (depending on the session size) and while a session is archiving, its database is locked and cannot be restored until the archiving process complets. When the Editor reboots after a crash, it searches for a session to recover, but skip over any session that is mounted/locked assuming the session is concurrently used by a concurrent Editor process, potentially preventing it from restoring. The optimal way to work around this problem  is to skip the archiving step. Instead, the live session is never archived (saving a copy), which allows the recovery service to shutdown and release the session lock very quickly ensuring that the session will be unlocked when the Editor restarts. On Editor start, it a crashed session is found and the user decides to restore it, the live session is copied into a new live session.

This changelist also affect those other jira in the following ways:

  - This CL changes execution order to shut down the recovery service ASAP to release the lock, but the optimization above make it super fast, so the UI should always be shown in a timely manner.

  - This CL ensures the recovery service release the session lock faster than the next instance of the Editor can start.

  - This CL fixes an issues with endpoints timeout logic.

  - This CL added code to clean up the intermediate directory left over by crashed client.

diff --git a/Engine/Plugins/Developer/Concert/ConcertMain/Source/Concert/Private/ConcertClient.cpp b/Engine/Plugins/Developer/Concert/ConcertMain/Source/Concert/Private/ConcertClient.cpp
--- a/Engine/Plugins/Developer/Concert/ConcertMain/Source/Concert/Private/ConcertClient.cpp
+++ b/Engine/Plugins/Developer/Concert/ConcertMain/Source/Concert/Private/ConcertClient.cpp
@@ -235,12 +235,12 @@ private:
 					// We found the default session to restore, restore and join it.
 					if (DefaultSessionToRestoreId.IsValid())
 					{
-						FConcertRestoreSessionArgs RestoreSessionArgs;
+						FConcertCopySessionArgs RestoreSessionArgs;
 						RestoreSessionArgs.bAutoConnect = true;
 						RestoreSessionArgs.SessionId = DefaultSessionToRestoreId;
 						RestoreSessionArgs.SessionName = Settings->DefaultSessionName;
 						RestoreSessionArgs.ArchiveNameOverride = Settings->DefaultSaveSessionAs;
-						return Client->InternalRestoreSession(ServerEndpoint, RestoreSessionArgs, MoveTemp(AutoConnectionNotification)).Share();
+						return Client->InternalCopySession(ServerEndpoint, RestoreSessionArgs, /*bRestoreOnlyConstraint*/true, MoveTemp(AutoConnectionNotification)).Share();
 					}
 
 					// No session found to join or restore, so create a new one.
@@ -875,11 +875,24 @@ TFuture<EConcertResponseCode> FConcertClient::JoinSession(const FGuid& ServerAdm
 	return InternalJoinSession(ServerAdminEndpointId, SessionId);
 }
 
-TFuture<EConcertResponseCode> FConcertClient::RestoreSession(const FGuid& ServerAdminEndpointId, const FConcertRestoreSessionArgs& RestoreSessionArgs)
+TFuture<EConcertResponseCode> FConcertClient::RestoreSession(const FGuid& ServerAdminEndpointId, const FConcertCopySessionArgs& RestoreSessionArgs)
 {
-	// We don't want the client to get automatically reconnected to it's default session if something wrong happens
-	AutoConnection.Reset();
-	return InternalRestoreSession(ServerAdminEndpointId, RestoreSessionArgs);
+	// We don't want the client to get automatically reconnected to the default session if something wrong happens
+	if (RestoreSessionArgs.bAutoConnect)
+	{
+		AutoConnection.Reset();
+	}
+	return InternalCopySession(ServerAdminEndpointId, RestoreSessionArgs, /*bRestoreOnlyConstraint*/true);
+}
+
+TFuture<EConcertResponseCode> FConcertClient::CopySession(const FGuid& ServerAdminEndpointId, const FConcertCopySessionArgs& CopySessionArgs)
+{
+	// We don't want the client to get automatically reconnected to the default session if the copy/connect fails.
+	if (CopySessionArgs.bAutoConnect)
+	{
+		AutoConnection.Reset();
+	}
+	return InternalCopySession(ServerAdminEndpointId, CopySessionArgs, /*bRestoreOnlyConstraint*/false);
 }
 
 TFuture<EConcertResponseCode> FConcertClient::ArchiveSession(const FGuid& ServerAdminEndpointId, const FConcertArchiveSessionArgs& ArchiveSessionArgs)
@@ -1167,18 +1180,19 @@ TFuture<EConcertResponseCode> FConcertClient::InternalJoinSession(const FGuid& S
 	return PendingConnection->Execute(MoveTemp(ConnectionTasks), MoveTemp(OngoingNotification));
 }
 
-TFuture<EConcertResponseCode> FConcertClient::InternalRestoreSession(const FGuid& ServerAdminEndpointId, const FConcertRestoreSessionArgs& RestoreSessionArgs, TUniquePtr<FAsyncTaskNotification> OngoingNotification)
+TFuture<EConcertResponseCode> FConcertClient::InternalCopySession(const FGuid& ServerAdminEndpointId, const FConcertCopySessionArgs& CopySessionArgs, bool bRestoreOnlyConstraint, TUniquePtr<FAsyncTaskNotification> OngoingNotification)
 {
-	FConcertAdmin_RestoreSessionRequest RestoreSessionRequest;
-	RestoreSessionRequest.SessionId = RestoreSessionArgs.SessionId;
-	RestoreSessionRequest.SessionName = RestoreSessionArgs.SessionName;
-	RestoreSessionRequest.SessionFilter = RestoreSessionArgs.SessionFilter;
-	RestoreSessionRequest.OwnerClientInfo = ClientInfo;
-	RestoreSessionRequest.VersionInfo.Initialize();
+	FConcertAdmin_CopySessionRequest CopySessionRequest;
+	CopySessionRequest.SessionId = CopySessionArgs.SessionId;
+	CopySessionRequest.SessionName = CopySessionArgs.SessionName;
+	CopySessionRequest.SessionFilter = CopySessionArgs.SessionFilter;
+	CopySessionRequest.bRestoreOnly = bRestoreOnlyConstraint;
+	CopySessionRequest.OwnerClientInfo = ClientInfo;
+	CopySessionRequest.VersionInfo.Initialize();
 
 	// Session settings
-	RestoreSessionRequest.SessionSettings.Initialize();
-	RestoreSessionRequest.SessionSettings.ArchiveNameOverride = RestoreSessionArgs.ArchiveNameOverride;
+	CopySessionRequest.SessionSettings.Initialize();
+	CopySessionRequest.SessionSettings.ArchiveNameOverride = CopySessionArgs.ArchiveNameOverride;
 
 	TUniquePtr<FAsyncTaskNotification> Notification = MoveTemp(OngoingNotification);
 	if (!Notification.IsValid())
@@ -1186,13 +1200,13 @@ TFuture<EConcertResponseCode> FConcertClient::InternalRestoreSession(const FGuid
 		FAsyncTaskNotificationConfig NotificationConfig;
 		NotificationConfig.bIsHeadless = Settings->bIsHeadless;
 		NotificationConfig.bKeepOpenOnFailure = true;
-		NotificationConfig.TitleText = LOCTEXT("RestoringSession", "Restoring Session...");
+		NotificationConfig.TitleText = bRestoreOnlyConstraint ? LOCTEXT("RestoringSession", "Restoring Session...") : LOCTEXT("CopyingSession", "Copying Session...");
 		NotificationConfig.LogCategory = ConcertUtil::GetLogConcertPtr();
 		Notification = MakeUnique<FAsyncTaskNotification>(NotificationConfig);
 	}
 
-	return ClientAdminEndpoint->SendRequest<FConcertAdmin_RestoreSessionRequest, FConcertAdmin_SessionInfoResponse>(RestoreSessionRequest, ServerAdminEndpointId)
-		.Next([this, Notification = MoveTemp(Notification), bAutoConnect = RestoreSessionArgs.bAutoConnect](const FConcertAdmin_SessionInfoResponse& RequestResponse) mutable
+	return ClientAdminEndpoint->SendRequest<FConcertAdmin_CopySessionRequest, FConcertAdmin_SessionInfoResponse>(CopySessionRequest, ServerAdminEndpointId)
+		.Next([this, Notification = MoveTemp(Notification), bAutoConnect = CopySessionArgs.bAutoConnect, bRestoreOnly = bRestoreOnlyConstraint](const FConcertAdmin_SessionInfoResponse& RequestResponse) mutable
 	{
 		if (RequestResponse.ResponseCode == EConcertResponseCode::Success)
 		{
@@ -1202,12 +1216,12 @@ TFuture<EConcertResponseCode> FConcertClient::InternalRestoreSession(const FGuid
 			}
 			else
 			{
-				Notification->SetComplete(FText::Format(LOCTEXT("RestoredSessionFmt", "Restored Session '{0}'"), FText::FromString(RequestResponse.SessionInfo.SessionName)), FText(), true);
+				Notification->SetComplete(FText::Format(bRestoreOnly ? LOCTEXT("RestoreSessionFmt", "Restored Session '{0}'") : LOCTEXT("CopySessionFmt", "Copied Session '{0}'"), FText::FromString(RequestResponse.SessionInfo.SessionName)), FText(), true);
 			}
 		}
 		else
 		{
-			Notification->SetComplete(LOCTEXT("FailedToRestoreSession", "Failed to Restore Session"), RequestResponse.Reason, false);
+			Notification->SetComplete(bRestoreOnly ? LOCTEXT("FailedToRestoreSession", "Failed to Restore Session") : LOCTEXT("FailedToCopySession", "Failed to Copy Session"), RequestResponse.Reason, false);
 		}
 		return RequestResponse.ResponseCode;
 	});
diff --git a/Engine/Plugins/Developer/Concert/ConcertMain/Source/Concert/Private/ConcertClient.h b/Engine/Plugins/Developer/Concert/ConcertMain/Source/Concert/Private/ConcertClient.h
--- a/Engine/Plugins/Developer/Concert/ConcertMain/Source/Concert/Private/ConcertClient.h
+++ b/Engine/Plugins/Developer/Concert/ConcertMain/Source/Concert/Private/ConcertClient.h
@@ -78,7 +78,8 @@ public:
 	virtual EConcertConnectionStatus GetSessionConnectionStatus() const override;
 	virtual TFuture<EConcertResponseCode> CreateSession(const FGuid& ServerAdminEndpointId, const FConcertCreateSessionArgs& CreateSessionArgs) override;
 	virtual TFuture<EConcertResponseCode> JoinSession(const FGuid& ServerAdminEndpointId, const FGuid& SessionId) override;
-	virtual TFuture<EConcertResponseCode> RestoreSession(const FGuid& ServerAdminEndpointId, const FConcertRestoreSessionArgs& RestoreSessionArgs) override;
+	virtual TFuture<EConcertResponseCode> RestoreSession(const FGuid& ServerAdminEndpointId, const FConcertCopySessionArgs& RestoreSessionArgs) override;
+	virtual TFuture<EConcertResponseCode> CopySession(const FGuid& ServerAdminEndpointId, const FConcertCopySessionArgs& CopySessionArgs) override;
 	virtual TFuture<EConcertResponseCode> ArchiveSession(const FGuid& ServerAdminEndpointId, const FConcertArchiveSessionArgs& ArchiveSessionArgs) override;
 	virtual TFuture<EConcertResponseCode> RenameSession(const FGuid& ServerAdminEndpointId, const FGuid& SessionId, const FString& NewName) override;
 	virtual TFuture<EConcertResponseCode> DeleteSession(const FGuid& ServerAdminEndpointId, const FGuid& SessionId) override;
@@ -109,7 +110,7 @@ private:
 
 	TFuture<EConcertResponseCode> InternalCreateSession(const FGuid& ServerAdminEndpointId, const FConcertCreateSessionArgs& CreateSessionArgs, TUniquePtr<FAsyncTaskNotification> OngoingNotification = nullptr);
 	TFuture<EConcertResponseCode> InternalJoinSession(const FGuid& ServerAdminEndpointId, const FGuid& SessionId, TUniquePtr<FAsyncTaskNotification> OngoingNotification = nullptr);
-	TFuture<EConcertResponseCode> InternalRestoreSession(const FGuid& ServerAdminEndpointId, const FConcertRestoreSessionArgs& RestoreSessionArgs, TUniquePtr<FAsyncTaskNotification> OngoingNotification = nullptr);
+	TFuture<EConcertResponseCode> InternalCopySession(const FGuid& ServerAdminEndpointId, const FConcertCopySessionArgs& CopySessionArgs, bool bRestoreOnlyConstraint, TUniquePtr<FAsyncTaskNotification> OngoingNotification = nullptr);
 	void InternalDisconnectSession();
 
 	/** Set the connection error. */
diff --git a/Engine/Plugins/Developer/Concert/ConcertMain/Source/Concert/Private/ConcertServer.cpp b/Engine/Plugins/Developer/Concert/ConcertMain/Source/Concert/Private/ConcertServer.cpp
--- a/Engine/Plugins/Developer/Concert/ConcertMain/Source/Concert/Private/ConcertServer.cpp
+++ b/Engine/Plugins/Developer/Concert/ConcertMain/Source/Concert/Private/ConcertServer.cpp
@@ -167,7 +167,7 @@ void FConcertServer::Startup()
 		// Add Session connection handling
 		ServerAdminEndpoint->RegisterRequestHandler<FConcertAdmin_CreateSessionRequest, FConcertAdmin_SessionInfoResponse>(this, &FConcertServer::HandleCreateSessionRequest);
 		ServerAdminEndpoint->RegisterRequestHandler<FConcertAdmin_FindSessionRequest, FConcertAdmin_SessionInfoResponse>(this, &FConcertServer::HandleFindSessionRequest);
-		ServerAdminEndpoint->RegisterRequestHandler<FConcertAdmin_RestoreSessionRequest, FConcertAdmin_SessionInfoResponse>(this, &FConcertServer::HandleRestoreSessionRequest);
+		ServerAdminEndpoint->RegisterRequestHandler<FConcertAdmin_CopySessionRequest, FConcertAdmin_SessionInfoResponse>(this, &FConcertServer::HandleCopySessionRequest);
 		ServerAdminEndpoint->RegisterRequestHandler<FConcertAdmin_ArchiveSessionRequest, FConcertAdmin_ArchiveSessionResponse>(this, &FConcertServer::HandleArchiveSessionRequest);
 		ServerAdminEndpoint->RegisterRequestHandler<FConcertAdmin_RenameSessionRequest, FConcertAdmin_RenameSessionResponse>(this, &FConcertServer::HandleRenameSessionRequest);
 		ServerAdminEndpoint->RegisterRequestHandler<FConcertAdmin_DeleteSessionRequest, FConcertAdmin_DeleteSessionResponse>(this, &FConcertServer::HandleDeleteSessionRequest);
@@ -265,7 +265,7 @@ void FConcertServer::Shutdown()
 		// Session connection
 		ServerAdminEndpoint->UnregisterRequestHandler<FConcertAdmin_CreateSessionRequest>();
 		ServerAdminEndpoint->UnregisterRequestHandler<FConcertAdmin_FindSessionRequest>();
-		ServerAdminEndpoint->UnregisterRequestHandler<FConcertAdmin_RestoreSessionRequest>();
+		ServerAdminEndpoint->UnregisterRequestHandler<FConcertAdmin_CopySessionRequest>();
 		ServerAdminEndpoint->UnregisterRequestHandler<FConcertAdmin_ArchiveSessionRequest>();
 		ServerAdminEndpoint->UnregisterRequestHandler<FConcertAdmin_RenameSessionRequest>();
 		ServerAdminEndpoint->UnregisterRequestHandler<FConcertAdmin_DeleteSessionRequest>();
@@ -284,14 +284,12 @@ void FConcertServer::Shutdown()
 
 	// Destroy the live sessions
 	{
-		const bool bAutoArchiveOnShutdown = true;
-
 		TArray<FGuid> LiveSessionIds;
 		LiveSessions.GetKeys(LiveSessionIds);
 		for (const FGuid& LiveSessionId : LiveSessionIds)
 		{
-			bool bDeleteSessionData = true;
-			if (bAutoArchiveOnShutdown)
+			bool bDeleteSessionData = false;
+			if (Settings->bAutoArchiveOnShutdown)
 			{
 				bDeleteSessionData = ArchiveLiveSession(LiveSessionId, FString(), AutoArchiveSessionFilter).IsValid();
 			}
@@ -414,35 +412,66 @@ TSharedPtr<IConcertServerSession> FConcertServer::CreateSession(const FConcertSe
 
 TSharedPtr<IConcertServerSession> FConcertServer::RestoreSession(const FGuid& SessionId, const FConcertSessionInfo& SessionInfo, const FConcertSessionFilter& SessionFilter, FText& OutFailureReason)
 {
-	if (!SessionInfo.SessionId.IsValid() || SessionInfo.SessionName.IsEmpty())
+	if (ArchivedSessions.Contains(SessionId))
 	{
-		OutFailureReason = LOCTEXT("Error_RestoreSession_EmptySessionIdOrName", "Empty session ID or name");
-		UE_LOG(LogConcert, Error, TEXT("An attempt to restore a session was made, but the session info was missing an ID or name!"));
-		return nullptr;
+		return CopySession(SessionId, SessionInfo, SessionFilter, OutFailureReason);
 	}
 
-	if (!Settings->ServerSettings.bIgnoreSessionSettingsRestriction && SessionInfo.VersionInfos.Num() == 0)
+	OutFailureReason = FText::Format(LOCTEXT("Error_RestoreSession_NotFound", "Session '{0}' not found"), FText::AsCultureInvariant(SessionId.ToString()));
+	UE_LOG(LogConcert, Error, TEXT("An attempt to restore session '%s' was made, but that session could not be found!"), *SessionId.ToString());
+	return nullptr;
+}
+
+TSharedPtr<IConcertServerSession> FConcertServer::CopySession(const FGuid& SrcSessionId, const FConcertSessionInfo& NewSessionInfo, const FConcertSessionFilter& SessionFilter, FText& OutFailureReason)
+{
+	if (!NewSessionInfo.SessionId.IsValid() || NewSessionInfo.SessionName.IsEmpty())
 	{
-		OutFailureReason = LOCTEXT("Error_RestoreSession_EmptyVersionInfo", "Empty version info");
-		UE_LOG(LogConcert, Error, TEXT("An attempt to restore a session was made, but the session info was missing version info!"));
+		OutFailureReason = LOCTEXT("Error_CopySession_EmptySessionIdOrName", "Empty session ID or name");
+		UE_LOG(LogConcert, Error, TEXT("An attempt to copy a session was made, but the session info was missing an ID or name!"));
 		return nullptr;
 	}
-
-	if (LiveSessions.Contains(SessionInfo.SessionId))
+	else if (!Settings->ServerSettings.bIgnoreSessionSettingsRestriction && NewSessionInfo.VersionInfos.Num() == 0)
 	{
-		OutFailureReason = FText::Format(LOCTEXT("Error_RestoreSession_AlreadyExists", "Session '{0}' already exists"), FText::AsCultureInvariant(SessionInfo.SessionId.ToString()));
-		UE_LOG(LogConcert, Error, TEXT("An attempt to restore a session with ID '%s' was made, but that session already exists!"), *SessionInfo.SessionId.ToString());
+		OutFailureReason = LOCTEXT("Error_CopySession_EmptyVersionInfo", "Empty version info");
+		UE_LOG(LogConcert, Error, TEXT("An attempt to copy a session was made, but the session info was missing version info!"));
 		return nullptr;
 	}
-
-	if (GetLiveSessionIdByName(SessionInfo.SessionName).IsValid())
+	else if (LiveSessions.Contains(NewSessionInfo.SessionId))
 	{
-		OutFailureReason = FText::Format(LOCTEXT("Error_RestoreSession_AlreadyExists", "Session '{0}' already exists"), FText::AsCultureInvariant(SessionInfo.SessionName));
-		UE_LOG(LogConcert, Error, TEXT("An attempt to restore a session with name '%s' was made, but that session already exists!"), *SessionInfo.SessionName);
+		OutFailureReason = FText::Format(LOCTEXT("Error_CopySession_AlreadyExists", "Session '{0}' already exists"), FText::AsCultureInvariant(NewSessionInfo.SessionId.ToString()));
+		UE_LOG(LogConcert, Error, TEXT("An attempt to copy a session with ID '%s' was made, but that session already exists!"), *NewSessionInfo.SessionId.ToString());
 		return nullptr;
 	}
+	else if (GetLiveSessionIdByName(NewSessionInfo.SessionName).IsValid())
+	{
+		OutFailureReason = FText::Format(LOCTEXT("Error_CopySession_AlreadyExists", "Session '{0}' already exists"), FText::AsCultureInvariant(NewSessionInfo.SessionName));
+		UE_LOG(LogConcert, Error, TEXT("An attempt to copy a session with name '%s' was made, but that session already exists!"), *NewSessionInfo.SessionName);
+		return nullptr;
+	}
+	else if (ArchivedSessions.Contains(SrcSessionId))
+	{
+		return RestoreArchivedSession(SrcSessionId, NewSessionInfo, SessionFilter, OutFailureReason);
+	}
+	else if (TSharedPtr<IConcertServerSession> LiveSession = LiveSessions.FindRef(SrcSessionId))
+	{
+		// Copy the live session in the default repository (where new sessions should be created), unless it is unset.
+		const FConcertSessionInfo& LiveSessionInfo = LiveSession->GetSessionInfo();
+		const FConcertServerSessionRepository& CopySessionRepository = DefaultSessionRepository.IsSet() ? DefaultSessionRepository.GetValue() : GetSessionRepository(LiveSession->GetSessionInfo().SessionId);
+		if (EventSink->CopySession(*this, LiveSession.ToSharedRef(), CopySessionRepository.GetSessionWorkingDir(NewSessionInfo.SessionId), SessionFilter))
+		{
+			UE_LOG(LogConcert, Display, TEXT("Live session '%s' (%s) was copied as '%s' (%s)"), *LiveSessionInfo.SessionName, *LiveSessionInfo.SessionId.ToString(), *NewSessionInfo.SessionName, *NewSessionInfo.SessionId.ToString());
+			return CreateLiveSession(NewSessionInfo, CopySessionRepository);
+		}
+		else
+		{
+			UE_LOG(LogConcert, Error, TEXT("An attempt to copy a session '%s' was made, but failed!"), *SrcSessionId.ToString());
+			return nullptr;
+		}
+	}
 
-	return RestoreArchivedSession(SessionId, SessionInfo, SessionFilter, OutFailureReason);
+	OutFailureReason = FText::Format(LOCTEXT("Error_CopySession_NotFound", "Session '{0}' not found"), FText::AsCultureInvariant(SrcSessionId.ToString()));
+	UE_LOG(LogConcert, Error, TEXT("An attempt to copy a session '%s' was made, but that session could not be found!"), *SrcSessionId.ToString());
+	return nullptr;
 }
 
 void FConcertServer::RecoverSessions(const FConcertServerSessionRepository& InRepository, bool bCleanupExpiredSessions)
@@ -1017,12 +1046,12 @@ TFuture<FConcertAdmin_SessionInfoResponse> FConcertServer::HandleFindSessionRequ
 	return FConcertAdmin_SessionInfoResponse::AsFuture(MoveTemp(ResponseData));
 }
 
-TFuture<FConcertAdmin_SessionInfoResponse> FConcertServer::HandleRestoreSessionRequest(const FConcertMessageContext& Context)
+TFuture<FConcertAdmin_SessionInfoResponse> FConcertServer::HandleCopySessionRequest(const FConcertMessageContext& Context)
 {
-	const FConcertAdmin_RestoreSessionRequest* Message = Context.GetMessage<FConcertAdmin_RestoreSessionRequest>();
+	const FConcertAdmin_CopySessionRequest* Message = Context.GetMessage<FConcertAdmin_CopySessionRequest>();
 
 	// Restore the server session
-	FText RestoreFailureReason;
+	FText FailureReason;
 	TSharedPtr<IConcertServerSession> NewServerSession;
 	{
 		FConcertSessionInfo SessionInfo = CreateSessionInfo();
@@ -1032,7 +1061,9 @@ TFuture<FConcertAdmin_SessionInfoResponse> FConcertServer::HandleRestoreSessionR
 		SessionInfo.SessionName = Message->SessionName;
 		SessionInfo.Settings = Message->SessionSettings;
 		SessionInfo.VersionInfos.Add(Message->VersionInfo);
-		NewServerSession = RestoreSession(Message->SessionId, SessionInfo, Message->SessionFilter, RestoreFailureReason);
+		NewServerSession = Message->bRestoreOnly ?
+			RestoreSession(Message->SessionId, SessionInfo, Message->SessionFilter, FailureReason) :
+			CopySession(Message->SessionId, SessionInfo, Message->SessionFilter, FailureReason);
 	}
 
 	// We have a valid session if it succeeded
@@ -1045,8 +1076,8 @@ TFuture<FConcertAdmin_SessionInfoResponse> FConcertServer::HandleRestoreSessionR
 	else
 	{
 		ResponseData.ResponseCode = EConcertResponseCode::Failed;
-		ResponseData.Reason = RestoreFailureReason;
-		UE_LOG(LogConcert, Display, TEXT("Session restoration failed. (User: %s, Reason: %s)"), *Message->OwnerClientInfo.UserName, *ResponseData.Reason.ToString());
+		ResponseData.Reason = FailureReason;
+		UE_LOG(LogConcert, Display, TEXT("Session copy failed. (User: %s, Reason: %s)"), *Message->OwnerClientInfo.UserName, *ResponseData.Reason.ToString());
 	}
 
 	return FConcertAdmin_SessionInfoResponse::AsFuture(MoveTemp(ResponseData));
diff --git a/Engine/Plugins/Developer/Concert/ConcertMain/Source/Concert/Private/ConcertServer.h b/Engine/Plugins/Developer/Concert/ConcertMain/Source/Concert/Private/ConcertServer.h
--- a/Engine/Plugins/Developer/Concert/ConcertMain/Source/Concert/Private/ConcertServer.h
+++ b/Engine/Plugins/Developer/Concert/ConcertMain/Source/Concert/Private/ConcertServer.h
@@ -39,6 +39,7 @@ public:
 	virtual TSharedPtr<IConcertServerSession> GetSession(const FGuid& SessionId) const override;
 	virtual TSharedPtr<IConcertServerSession> CreateSession(const FConcertSessionInfo& SessionInfo, FText& OutFailureReason) override;
 	virtual TSharedPtr<IConcertServerSession> RestoreSession(const FGuid& SessionId, const FConcertSessionInfo& SessionInfo, const FConcertSessionFilter& SessionFilter, FText& OutFailureReason) override;
+	virtual TSharedPtr<IConcertServerSession> CopySession(const FGuid& SrcSessionId, const FConcertSessionInfo& NewSessionInfo, const FConcertSessionFilter& SessionFilter, FText& OutFailureReason) override;
 	virtual FGuid ArchiveSession(const FGuid& SessionId, const FString& ArchiveNameOverride, const FConcertSessionFilter& SessionFilter, FText& OutFailureReason) override;
 	virtual bool ExportSession(const FGuid& SessionId, const FConcertSessionFilter& SessionFilter, const FString& DestDir, bool bAnonymizeData, FText& OutFailureReason) override;
 	virtual bool RenameSession(const FGuid& SessionId, const FString& NewName, FText& OutFailureReason) override;
@@ -75,7 +76,7 @@ private:
 	TFuture<FConcertAdmin_SessionInfoResponse> HandleFindSessionRequest(const FConcertMessageContext& Context);
 
 	/**  */
-	TFuture<FConcertAdmin_SessionInfoResponse> HandleRestoreSessionRequest(const FConcertMessageContext& Context);
+	TFuture<FConcertAdmin_SessionInfoResponse> HandleCopySessionRequest(const FConcertMessageContext& Context);
 
 	/**  */
 	TFuture<FConcertAdmin_ArchiveSessionResponse> HandleArchiveSessionRequest(const FConcertMessageContext& Context);
diff --git a/Engine/Plugins/Developer/Concert/ConcertMain/Source/Concert/Public/ConcertMessages.h b/Engine/Plugins/Developer/Concert/ConcertMain/Source/Concert/Public/ConcertMessages.h
--- a/Engine/Plugins/Developer/Concert/ConcertMain/Source/Concert/Public/ConcertMessages.h
+++ b/Engine/Plugins/Developer/Concert/ConcertMain/Source/Concert/Public/ConcertMessages.h
@@ -272,34 +272,38 @@ struct FConcertAdmin_FindSessionRequest : public FConcertRequestData
 	FConcertSessionVersionInfo VersionInfo;
 };
 
+/** Used to copy a live session or restore an archived one. */
 USTRUCT()
-struct FConcertAdmin_RestoreSessionRequest : public FConcertRequestData
+struct FConcertAdmin_CopySessionRequest : public FConcertRequestData
 {
 	GENERATED_BODY()
 
-	/** The ID of the session to restore (must be an archived session). */
+	/** The ID of the session to copy or restore. It can be a live or an archived session unless bRestoreOnly is true, in which case, the session to copy must be an archived one. */
 	UPROPERTY(VisibleAnywhere, Category="Concert Message")
 	FGuid SessionId;
 
-	/** The name of the restored session to create. */
+	/** The name of the session to create. */
 	UPROPERTY(VisibleAnywhere, Category="Concert Message")
 	FString SessionName;
 
-	/** Information about the owner of the restored session. */
+	/** Information about the owner of the copied session. */
 	UPROPERTY(VisibleAnywhere, Category="Concert Message")
 	FConcertClientInfo OwnerClientInfo;
 
-	/** Settings to apply to the restored session. */
+	/** Settings to apply to the copied session. */
 	UPROPERTY(VisibleAnywhere, Category="Concert Message")
 	FConcertSessionSettings SessionSettings;
 
-	/** Version information of the client requesting the restore. */
+	/** Version information of the client requesting the copy. */
 	UPROPERTY(VisibleAnywhere, Category="Concert Message")
 	FConcertSessionVersionInfo VersionInfo;
 
-	/** The filter controlling which activities from the session should be restored. */
+	/** The filter controlling which activities from the session should be copied over. */
 	UPROPERTY(VisibleAnywhere, Category="Concert Message")
 	FConcertSessionFilter SessionFilter;
+
+	/** True to constrain the session to copy to be an archive (implying a restore operation). */
+	bool bRestoreOnly = false;
 };
 
 USTRUCT()
diff --git a/Engine/Plugins/Developer/Concert/ConcertMain/Source/Concert/Public/ConcertSettings.h b/Engine/Plugins/Developer/Concert/ConcertMain/Source/Concert/Public/ConcertSettings.h
--- a/Engine/Plugins/Developer/Concert/ConcertMain/Source/Concert/Public/ConcertSettings.h
+++ b/Engine/Plugins/Developer/Concert/ConcertMain/Source/Concert/Public/ConcertSettings.h
@@ -107,6 +107,12 @@ public:
 	UPROPERTY(config)
 	bool bAutoArchiveOnReboot = false;
 
+	/**
+	 * If true, instruct the server to auto-archive live sessions on shutdown.
+	 */
+	UPROPERTY(config)
+	bool bAutoArchiveOnShutdown = true;
+
 	/**
 	 * Clean server sessions working directory when booting
 	 * Can be specified on the server cmd with `-CONCERTCLEAN`
diff --git a/Engine/Plugins/Developer/Concert/ConcertMain/Source/Concert/Public/IConcertClient.h b/Engine/Plugins/Developer/Concert/ConcertMain/Source/Concert/Public/IConcertClient.h
--- a/Engine/Plugins/Developer/Concert/ConcertMain/Source/Concert/Public/IConcertClient.h
+++ b/Engine/Plugins/Developer/Concert/ConcertMain/Source/Concert/Public/IConcertClient.h
@@ -113,21 +113,21 @@ struct FConcertCreateSessionArgs
 	FString ArchiveNameOverride;
 };
 
-struct FConcertRestoreSessionArgs
+struct FConcertCopySessionArgs
 {
-	/** True to auto-connect to the session after restoring it */
+	/** True to auto-connect to the session after copying/restoring it */
 	bool bAutoConnect = true;
 
-	/** The ID of the archived session to restore */
+	/** The ID of the session to copy or restore */
 	FGuid SessionId;
 
 	/** The desired name for new session */
 	FString SessionName;
 
-	/** The override for the name used when archiving this session */
+	/** The override for the name used when archiving the copied/restored session */
 	FString ArchiveNameOverride;
 
-	/** The filter controlling which activities from the session should be restored */
+	/** The filter controlling which activities should be copied from the source session. */
 	FConcertSessionFilter SessionFilter;
 };
 
@@ -292,13 +292,23 @@ public:
 	virtual TFuture<EConcertResponseCode> JoinSession(const FGuid& ServerAdminEndpointId, const FGuid& SessionId) = 0;
 
 	/**
-	 * Restore an archived session on the server, matching the client configured settings.
+	 * Copy an archived session into a new live session on the server, matching the client configured settings.
 	 * This also initiates the connection handshake for that session with the client when bAutoConnect is true in the RestoreSessionArgs.
 	 * @param ServerAdminEndpointId	The Id of the Concert Server query endpoint
 	 * @param RestoreSessionArgs	The arguments that will be use for the restoration of the session
 	 * @return A future that will contains the final response code of the request
 	 */
-	virtual TFuture<EConcertResponseCode> RestoreSession(const FGuid& ServerAdminEndpointId, const FConcertRestoreSessionArgs& RestoreSessionArgs) = 0;
+	virtual TFuture<EConcertResponseCode> RestoreSession(const FGuid& ServerAdminEndpointId, const FConcertCopySessionArgs& RestoreSessionArgs) = 0;
+
+	/**
+	 * Copy a live or and archived session into a new live session on server, matching the configured settings. If the session is archived,
+	 * this is equivalent to restoring it (because copying an archive into another archive is not useful).
+	 * This also initiates the connection handshake for that session with the server when bAutoConnect is true in the CopySessionArgs.
+	 * @param ServerAdminEndpointId	The Id of the Concert Server query endpoint
+	 * @param CopySessionArgs	The arguments that will be use to copy the session
+	 * @return A future that will contains the final response code of the request
+	 */
+	virtual TFuture<EConcertResponseCode> CopySession(const FGuid& ServerAdminEndpointId, const FConcertCopySessionArgs& CopySessionArgs) = 0;
 
 	/**
 	 * Archive a live session on the server hosting the session.
diff --git a/Engine/Plugins/Developer/Concert/ConcertMain/Source/Concert/Public/IConcertServer.h b/Engine/Plugins/Developer/Concert/ConcertMain/Source/Concert/Public/IConcertServer.h
--- a/Engine/Plugins/Developer/Concert/ConcertMain/Source/Concert/Public/IConcertServer.h
+++ b/Engine/Plugins/Developer/Concert/ConcertMain/Source/Concert/Public/IConcertServer.h
@@ -99,6 +99,17 @@ public:
 	 */
 	virtual TSharedPtr<IConcertServerSession> CreateSession(const FConcertSessionInfo& SessionInfo, FText& OutFailureReason) = 0;
 
+	/** 
+	 * Create a new live server session from another session. The source session can be an archive or a live session.
+	 * @param SrcSessionId The ID of the session to copy. If the session is archived, its equivalent to restore it.
+	 * @param NewSessionInfo The information about the new session to create.
+	 * @param SessionFilter The filter controlling which activities should be copied over the new session.
+	 * @param OutFailureReason The reason the operation fails if the function returns null, undefined otherwise.
+	 * @return the created server session
+	 * @note This is equivalent to archiving and restoring a session, but faster as it skips one copy of the session (possibly serveral GB).
+	 */
+	virtual TSharedPtr<IConcertServerSession> CopySession(const FGuid& SrcSessionId, const FConcertSessionInfo& NewSessionInfo, const FConcertSessionFilter& SessionFilter, FText& OutFailureReason) = 0;
+
 	/**
 	 * Restore an archived Concert server session based on the passed session info
 	 * @param SessionId The ID of the session to restore
diff --git a/Engine/Plugins/Developer/Concert/ConcertMain/Source/Concert/Public/IConcertServerEventSink.h b/Engine/Plugins/Developer/Concert/ConcertMain/Source/Concert/Public/IConcertServerEventSink.h
--- a/Engine/Plugins/Developer/Concert/ConcertMain/Source/Concert/Public/IConcertServerEventSink.h
+++ b/Engine/Plugins/Developer/Concert/ConcertMain/Source/Concert/Public/IConcertServerEventSink.h
@@ -56,6 +56,11 @@ public:
 	 */
 	virtual bool ArchiveSession(const IConcertServer& InServer, const FString& InLiveSessionWorkingDir, const FString& InArchivedSessionRoot, const FConcertSessionInfo& InArchivedSessionInfo, const FConcertSessionFilter& InSessionFilter) = 0;
 
+	/**
+	 * Called to copy the data of a live session into another live session.
+	 */
+	virtual bool CopySession(const IConcertServer& InServer, TSharedRef<IConcertServerSession> InLiveSession, const FString& NewSessionRoot, const FConcertSessionFilter& InSessionFilter) = 0;
+
 	/**
 	 * Called to migrate and gather the data of a live or archived session to a directory for external usage.
 	 * @note As opposed to archiving a session, once exported, the session is not tracked anymore by the server and the exported files can be copied/modified while the server is running.
diff --git a/Engine/Plugins/Developer/Concert/ConcertMain/Source/ConcertTransport/Private/ConcertLocalEndpoint.cpp b/Engine/Plugins/Developer/Concert/ConcertMain/Source/ConcertTransport/Private/ConcertLocalEndpoint.cpp
--- a/Engine/Plugins/Developer/Concert/ConcertMain/Source/ConcertTransport/Private/ConcertLocalEndpoint.cpp
+++ b/Engine/Plugins/Developer/Concert/ConcertMain/Source/ConcertTransport/Private/ConcertLocalEndpoint.cpp
@@ -735,8 +735,9 @@ void FConcertLocalEndpoint::SendKeepAlives(const FDateTime& UtcNow)
 		FConcertRemoteEndpointPtr RemoteEndpoint = RemoteEndpointPair.Value;
 		check(RemoteEndpoint.IsValid());
 
-		// if no message have been sent to this endpoint for a quarter of the timeout span, send a keep alive
-		if (!RemoteEndpoint->GetEndpointTimeoutSpan().IsZero() && RemoteEndpoint->GetLastSentMessageTime() + (RemoteEndpoint->GetEndpointTimeoutSpan() * 0.25f) <= UtcNow)
+		// if no message have been sent to this endpoint for some time, send a keep alive to ensure the UDP nodes re-registered if MessageBus decided to unregister them.
+		FTimespan KeepAliveSpan = RemoteEndpoint->GetEndpointTimeoutSpan().IsZero() ? FTimespan(0, 0, 10) : RemoteEndpoint->GetEndpointTimeoutSpan() * 0.25f;
+		if (RemoteEndpoint->GetLastSentMessageTime() + KeepAliveSpan <= UtcNow)
 		{
 			SendKeepAlive(RemoteEndpoint.ToSharedRef(), UtcNow);
 		}
diff --git a/Engine/Plugins/Developer/Concert/ConcertSync/ConcertSyncServer/Source/ConcertSyncServer/Private/ConcertSyncServer.cpp b/Engine/Plugins/Developer/Concert/ConcertSync/ConcertSyncServer/Source/ConcertSyncServer/Private/ConcertSyncServer.cpp
--- a/Engine/Plugins/Developer/Concert/ConcertSync/ConcertSyncServer/Source/ConcertSyncServer/Private/ConcertSyncServer.cpp
+++ b/Engine/Plugins/Developer/Concert/ConcertSync/ConcertSyncServer/Source/ConcertSyncServer/Private/ConcertSyncServer.cpp
@@ -285,11 +285,7 @@ void FConcertSyncServer::OnArchivedSessionDestroyed(const IConcertServer& InServ
 
 bool FConcertSyncServer::ArchiveSession(const IConcertServer& InServer, TSharedRef<IConcertServerSession> InLiveSession, const FString& InArchivedSessionRoot, const FConcertSessionInfo& InArchivedSessionInfo, const FConcertSessionFilter& InSessionFilter)
 {
-	if (TSharedPtr<FConcertSyncServerLiveSession> LiveSession = LiveSessions.FindRef(InLiveSession->GetId()))
-	{
-		return ConcertSyncServerUtils::MigrateSessionData(LiveSession->GetSessionDatabase(), InArchivedSessionRoot, InSessionFilter);
-	}
-	return false;
+	return CopySession(InServer, InLiveSession, InArchivedSessionRoot, InSessionFilter);
 }
 
 bool FConcertSyncServer::ArchiveSession(const IConcertServer& InServer, const FString& InLiveSessionWorkingDir, const FString& InArchivedSessionRoot, const FConcertSessionInfo& InArchivedSessionInfo, const FConcertSessionFilter& InSessionFilter)
@@ -304,6 +300,15 @@ bool FConcertSyncServer::ArchiveSession(const IConcertServer& InServer, const FS
 	return RetVal;
 }
 
+bool FConcertSyncServer::CopySession(const IConcertServer& InServer, TSharedRef<IConcertServerSession> InLiveSession, const FString& NewSessionRoot, const FConcertSessionFilter& InSessionFilter)
+{
+	if (TSharedPtr<FConcertSyncServerLiveSession> LiveSession = LiveSessions.FindRef(InLiveSession->GetId()))
+	{
+		return ConcertSyncServerUtils::MigrateSessionData(LiveSession->GetSessionDatabase(), NewSessionRoot, InSessionFilter);
+	}
+	return false;
+}
+
 bool FConcertSyncServer::ExportSession(const IConcertServer& InServer, const FGuid& InSessionId, const FString& DestDir, const FConcertSessionFilter& InSessionFilter, bool bAnonymizeData)
 {
 	if (TSharedPtr<FConcertSyncServerLiveSession> LiveSession = LiveSessions.FindRef(InSessionId)) // If the session is live.
diff --git a/Engine/Plugins/Developer/Concert/ConcertSync/ConcertSyncServer/Source/ConcertSyncServer/Private/ConcertSyncServer.h b/Engine/Plugins/Developer/Concert/ConcertSync/ConcertSyncServer/Source/ConcertSyncServer/Private/ConcertSyncServer.h
--- a/Engine/Plugins/Developer/Concert/ConcertSync/ConcertSyncServer/Source/ConcertSyncServer/Private/ConcertSyncServer.h
+++ b/Engine/Plugins/Developer/Concert/ConcertSync/ConcertSyncServer/Source/ConcertSyncServer/Private/ConcertSyncServer.h
@@ -36,6 +36,7 @@ public:
 	virtual void OnArchivedSessionDestroyed(const IConcertServer& InServer, const FGuid& InArchivedSessionId) override;
 	virtual bool ArchiveSession(const IConcertServer& InServer, TSharedRef<IConcertServerSession> InLiveSession, const FString& InArchivedSessionRoot, const FConcertSessionInfo& InArchivedSessionInfo, const FConcertSessionFilter& InSessionFilter) override;
 	virtual bool ArchiveSession(const IConcertServer& InServer, const FString& InLiveSessionWorkingDir, const FString& InArchivedSessionRoot, const FConcertSessionInfo& InArchivedSessionInfo, const FConcertSessionFilter& InSessionFilter) override;
+	virtual bool CopySession(const IConcertServer& InServer, TSharedRef<IConcertServerSession> InLiveSession, const FString& NewSessionRoot, const FConcertSessionFilter& InSessionFilter) override;
 	virtual bool ExportSession(const IConcertServer& InServer, const FGuid& InSessionId, const FString& DestDir, const FConcertSessionFilter& InSessionFilter, bool bAnonymizeData) override;
 	virtual bool RestoreSession(const IConcertServer& InServer, const FGuid& InArchivedSessionId, const FString& InLiveSessionRoot, const FConcertSessionInfo& InLiveSessionInfo, const FConcertSessionFilter& InSessionFilter) override;
 	virtual bool GetSessionActivities(const IConcertServer& InServer, const FGuid& SessionId, int64 FromActivityId, int64 ActivityCount, TArray<FConcertSessionSerializedPayload>& OutActivities, TMap<FGuid, FConcertClientInfo>& OutEndpointClientInfoMap, bool bIncludeDetails) override;
diff --git a/Engine/Plugins/Developer/Concert/DisasterRecoveryClient/Source/DisasterRecoveryClient/Private/DisasterRecoveryClientModule.cpp b/Engine/Plugins/Developer/Concert/DisasterRecoveryClient/Source/DisasterRecoveryClient/Private/DisasterRecoveryClientModule.cpp
--- a/Engine/Plugins/Developer/Concert/DisasterRecoveryClient/Source/DisasterRecoveryClient/Private/DisasterRecoveryClientModule.cpp
+++ b/Engine/Plugins/Developer/Concert/DisasterRecoveryClient/Source/DisasterRecoveryClient/Private/DisasterRecoveryClientModule.cpp
@@ -100,7 +100,7 @@ public:
 		ClientConfig->DefaultSessionName = DisasterRecoverySessionName;
 		ClientConfig->DefaultSaveSessionAs = DisasterRecoverySessionName;
 		//ClientConfig->ClientSettings.DiscoveryTimeoutSeconds = 0; // Setting this to zero prevents showing 'Server {serverName} lost." but will not detect when a server dies, growing the list of known server indefinitely.
-		ClientConfig->EndpointSettings.RemoteEndpointTimeoutSeconds = 4; // Non-zero -> Ensure the Endpoints are re-registered with Message Bus if Message Bus restarts or some errors trigger a 'auto-repair'.
+		ClientConfig->EndpointSettings.RemoteEndpointTimeoutSeconds = 0; // Ensure the endpoints never time out (and are kept alive automatically by Concert).
 		ClientConfig->ClientSettings.ClientAuthenticationKey = DisasterRecoveryServerName; // The server adds its own server name to the list of authorized client keys, use that key to authorize this client on the server.
 
 		// Create and start the client.
diff --git a/Engine/Plugins/Developer/Concert/DisasterRecoveryClient/Source/DisasterRecoveryClient/Private/DisasterRecoverySessionInfo.h b/Engine/Plugins/Developer/Concert/DisasterRecoveryClient/Source/DisasterRecoveryClient/Private/DisasterRecoverySessionInfo.h
--- a/Engine/Plugins/Developer/Concert/DisasterRecoveryClient/Source/DisasterRecoveryClient/Private/DisasterRecoverySessionInfo.h
+++ b/Engine/Plugins/Developer/Concert/DisasterRecoveryClient/Source/DisasterRecoveryClient/Private/DisasterRecoverySessionInfo.h
@@ -46,11 +46,11 @@ struct FDisasterRecoverySession
 
 	/** The client that mounted that session repository. */
 	UPROPERTY()
-	int32 MountedByProcessId = 0;
+	uint32 MountedByProcessId = 0;
 
 	/** The process ID of the client connected to the session (the session is live and recording transactions for this client process). */
 	UPROPERTY()
-	int32 ClientProcessId = 0;
+	uint32 ClientProcessId = 0;
 
 	/** Information about the session. */
 	UPROPERTY()
@@ -95,6 +95,22 @@ struct FDisasterRecoverySession
 	}
 };
 
+/** Information about a disaster recovery client. */
+USTRUCT()
+struct FDisasterRecoveryClientInfo
+{
+	GENERATED_BODY()
+
+	/** The client process ID. */
+	UPROPERTY()
+	uint32 ClientProcessId;
+
+	/** The client app ID. */
+	UPROPERTY()
+	FGuid ClientAppId;
+};
+
+
 /**
  * Hold the information for multiple disaster recovery sessions.
  */
@@ -127,7 +143,7 @@ struct FDisasterRecoveryInfo
 	UPROPERTY()
 	TArray<FGuid> DiscardedRepositoryIds;
 
-	/** The list of process IDs for disaster recovery client currently executing. */
+	/** The list of client currently executing (in different processes). */
 	UPROPERTY()
-	TArray<int32> ClientProcessIds;
+	TArray<FDisasterRecoveryClientInfo> Clients;
 };
diff --git a/Engine/Plugins/Developer/Concert/DisasterRecoveryClient/Source/DisasterRecoveryClient/Private/DisasterRecoverySessionManager.cpp b/Engine/Plugins/Developer/Concert/DisasterRecoveryClient/Source/DisasterRecoveryClient/Private/DisasterRecoverySessionManager.cpp
--- a/Engine/Plugins/Developer/Concert/DisasterRecoveryClient/Source/DisasterRecoveryClient/Private/DisasterRecoverySessionManager.cpp
+++ b/Engine/Plugins/Developer/Concert/DisasterRecoveryClient/Source/DisasterRecoveryClient/Private/DisasterRecoverySessionManager.cpp
@@ -287,10 +287,6 @@ FDisasterRecoverySessionManager::~FDisasterRecoverySessionManager()
 
 	SyncClient->GetConcertClient()->OnSessionStartup().RemoveAll(this);
 	SyncClient->GetConcertClient()->OnSessionShutdown().RemoveAll(this);
-
-	// Remove this recovery clients from the list of active clients.
-	DisasterRecoveryUtil::ScopedRecoveryInfoExclusiveUpdater ScopedRecoveryUpdater(RecoveryInfoPathname, SessionsCacheRevision);
-	ScopedRecoveryUpdater.Info().ClientProcessIds.RemoveAll([](const int32 ProcessId) { return ProcessId == FPlatformProcess::GetCurrentProcessId(); });
 }
 
 TArray<FGuid> FDisasterRecoverySessionManager::InitAndRotateSessions()
@@ -306,12 +302,22 @@ TArray<FGuid> FDisasterRecoverySessionManager::InitAndRotateSessions()
 		IFileManager::Get().Delete(*OldSessionInfo);
 	}
 
-	// Remove dead recovery clients from the list of active clients.
-	RecoveryInfo.ClientProcessIds.RemoveAll([](const int32 ProcessId) { return !DisasterRecoveryUtil::IsRecoveryClientProcessRunning(ProcessId);});
+	// Remove dead/crashed recovery clients from the list of active clients.
+	RecoveryInfo.Clients.RemoveAll([this](const FDisasterRecoveryClientInfo& ClientInfo)
+	{
+		if (!DisasterRecoveryUtil::IsRecoveryClientProcessRunning(ClientInfo.ClientProcessId))
+		{
+			// Clean up the temporary directory left behind by this client (it can use lot of disk space).
+			FString ClientTempDir = FPaths::ProjectIntermediateDir() / TEXT("Concert") / Role / ClientInfo.ClientAppId.ToString(); // See FConcertClientPaths (inaccessible from here).
+			IFileManager::Get().DeleteDirectory(*ClientTempDir, /*bRequireExists*/false, /*Tree*/true);
+			return true;
+		}
+		return false; // The client is sill running, don't remove it.
+	});
 
 	// Add this recovery client to the list of active clients.
-	RecoveryInfo.ClientProcessIds.Add(FPlatformProcess::GetCurrentProcessId());
-	DisasterRecoveryClientCount = RecoveryInfo.ClientProcessIds.Num();
+	RecoveryInfo.Clients.Add(FDisasterRecoveryClientInfo{FPlatformProcess::GetCurrentProcessId(), FApp::GetInstanceId()});
+	DisasterRecoveryClientCount = RecoveryInfo.Clients.Num();
 
 	// Unmap repositories that doesn't exist anymore on disk. User might have manually deleted the directories.
 	auto UnmapDeletedRepositoriesFn = [&RecoveryInfo](TArray<FDisasterRecoverySession>& Sessions)
@@ -526,8 +532,8 @@ TFuture<bool> FDisasterRecoverySessionManager::HasRecoverableCandidates()
 			continue; // Don't prompt the user to recover when the debugger was attached. User likely stopped the debugging in the middle of a session and is likely not interested to recover.
 		}
 
-		// Unreviewed crashes are obvious candidates. For a live session, if its repository cannot be mounted by this process, it will be eliminated, but if can, it means that
-		// the 'out-of-process' crash reporter has caught a crash, archived and unmounted the session even if the owned process is still active from OS point of view.
+		// Unreviewed crashes are obvious candidates. For a live session (owner process is still alive), if its repository cannot be mounted by this process, it will be eliminated, but if can,
+		// it means that the 'out-of-process' crash reporter caught a crash and unmounted the session repository even if the owned process is still active from OS point of view.
 		if (Candidate->IsUnreviewedCrash() || (Candidate->IsLive() && Candidate->ClientProcessId != FPlatformProcess::GetCurrentProcessId()))
 		{
 			Candidates.Add(Candidate);
@@ -561,33 +567,33 @@ TFuture<TVariant<TSharedPtr<FConcertActivityStream>, FText>> FDisasterRecoverySe
 		Promise->EmplaceValue(TInPlaceType<FText>(), ErrorMsg); // Could not load the session activities.
 	};
 
-	auto OnArchiveIdLookupResultFn = [this, OnErrorFn, Session, Promise](const FGuid& ArchivedSessionId)
+	auto OnSessionIdLookupResultFn = [this, OnErrorFn, Session, Promise](const FGuid& ResultSessionId)
 	{
-		if (!ArchivedSessionId.IsValid()) // Not found?
+		if (!ResultSessionId.IsValid()) // Not found?
 		{
-			OnSessionNotFoundInternal(Session->RepositoryId); // The archive was likely deleted. Mark it to be discarded to be cleaned up at the next start up.
+			OnSessionNotFoundInternal(Session->RepositoryId); // The session was likely deleted. Mark it to be discarded to be cleaned up at the next start up.
 			OnErrorFn(LOCTEXT("SessionNotFound", "Session not found."));
 			return; // 3b - No session.
 		}
 
 		// 3a- Return the activity stream on the session.
-		TSharedPtr<FConcertActivityStream> ActivityStream = MakeShared<FConcertActivityStream>(SyncClient->GetConcertClient(), ServerAdminEndpointId, ArchivedSessionId, /*bInIncludeActivityDetails*/true);
+		TSharedPtr<FConcertActivityStream> ActivityStream = MakeShared<FConcertActivityStream>(SyncClient->GetConcertClient(), ServerAdminEndpointId, ResultSessionId, /*bInIncludeActivityDetails*/true);
 		return Promise->EmplaceValue(TInPlaceType<TSharedPtr<FConcertActivityStream>>(), ActivityStream);
 	};
 
-	auto OnRepositoryMountedFn = [this, Session, OnErrorFn, OnArchiveIdLookupResultFn]()
+	auto OnRepositoryMountedFn = [this, Session, OnErrorFn, OnSessionIdLookupResultFn]()
 	{
 		OnSessionRepositoryMountedInternal(Session->RepositoryId); // Update internal state.
 
-		// 2- Find the archived session by name and retrieve its ID.
+		// 2- Find the session by name and retrieve its ID.
 		TaskExecutor.Enqueue(MakeShared<FLookupDisasterRecoverySessionIdTask>(
 			[this](){ return SyncClient->GetConcertClient()->GetServerSessions(ServerAdminEndpointId); },
 			Session->SessionName,
-			OnArchiveIdLookupResultFn,
+			OnSessionIdLookupResultFn,
 			OnErrorFn));
 	};
 
-	// 1- Mount the session repository, the server will reload the archived session in memory. (Create the mount point if it doesn't exist in case its for an imported session the server has not seen yet)
+	// 1- Mount the session repository, the server will reload the session in memory. (Create the mount point if it doesn't exist in case its for an imported session the server has not seen yet)
 	TaskExecutor.Enqueue(MakeShared<FMountDisasterRecoverySessionRepositoryTask>(
 		[this, Session]() { return SyncClient->GetConcertClient()->MountSessionRepository(ServerAdminEndpointId, Session->RepositoryRootDir, Session->RepositoryId, /*bCreateIfNotExist*/true, /*bAsDefault*/false); },
 		OnRepositoryMountedFn,
@@ -807,7 +813,7 @@ TFuture<TPair<bool, FText>> FDisasterRecoverySessionManager::RestoreAndJoinSessi
 	TSharedPtr<TPromise<TPair<bool, FText>>> Promise = MakeShared<TPromise<TPair<bool, FText>>>();
 	FString SessionName = RecoveryService::MakeSessionName(); // This follow the convention used in CrachReportClientEditor.
 	FGuid RepositoryId = FGuid::NewGuid(); // The repository for this session.
-	TSharedPtr<FGuid> ArchivedSessionId = MakeShared<FGuid>(); // Will hold the ID of the session to restore (once found).
+	TSharedPtr<FGuid> SrcSessionId = MakeShared<FGuid>(); // Will hold the ID of the session to restore (once found).
 
 	// Remember the current session info.
 	CurrentSessionRepositoryId = RepositoryId;
@@ -844,7 +850,7 @@ TFuture<TPair<bool, FText>> FDisasterRecoverySessionManager::RestoreAndJoinSessi
 		}
 	};
 
-	auto OnNewSessionRepositoryMountedFn = [this, Promise, RepositoryId, ThroughActivity, SessionName, ArchivedSessionId, OnSessionRestoredFn, OnErrorFn]()
+	auto OnNewSessionRepositoryMountedFn = [this, Promise, RepositoryId, ThroughActivity, SessionName, SrcSessionId, OnSessionRestoredFn, OnErrorFn]()
 	{
 		OnSessionRepositoryMountedInternal(RepositoryId); // Update internal state.
 
@@ -854,25 +860,25 @@ TFuture<TPair<bool, FText>> FDisasterRecoverySessionManager::RestoreAndJoinSessi
 			return;
 		}
 
-		FConcertRestoreSessionArgs RestoreArgs;
-		RestoreArgs.ArchiveNameOverride = SessionName;
-		RestoreArgs.bAutoConnect = true; // Join the session once restored.
-		RestoreArgs.SessionId = *ArchivedSessionId;
-		RestoreArgs.SessionFilter.ActivityIdUpperBound = ThroughActivity->Activity.ActivityId;
-		RestoreArgs.SessionFilter.bOnlyLiveData = false; // Reapply all activities.
-		RestoreArgs.SessionFilter.bIncludeIgnoredActivities = false; // Don't restore ignored activities. (Like MultiUser activities recorded in a DisasterRecovery session)
-		RestoreArgs.SessionName = SessionName;
+		FConcertCopySessionArgs CopyArgs;
+		CopyArgs.ArchiveNameOverride = SessionName;
+		CopyArgs.bAutoConnect = true; // Join the session once copied.
+		CopyArgs.SessionId = *SrcSessionId; // The session to copy.
+		CopyArgs.SessionFilter.ActivityIdUpperBound = ThroughActivity->Activity.ActivityId;
+		CopyArgs.SessionFilter.bOnlyLiveData = false; // Reapply all activities.
+		CopyArgs.SessionFilter.bIncludeIgnoredActivities = false; // Don't restore ignored activities. (Like MultiUser activities recorded in a DisasterRecovery session)
+		CopyArgs.SessionName = SessionName;
 
-		// 4 - Create a new session from the archive (restore) and join it.
+		// 4 - Create a new session from the source one and join it.
 		TaskExecutor.Enqueue(MakeShared<FCreateAndJoinDisasterRecoverySessionTask>(
-			[this, RestoreArgs = MoveTemp(RestoreArgs)](){ return SyncClient->GetConcertClient()->RestoreSession(ServerAdminEndpointId, RestoreArgs); },
+			[this, CopyArgs = MoveTemp(CopyArgs)](){ return SyncClient->GetConcertClient()->CopySession(ServerAdminEndpointId, CopyArgs); },
 			SyncClient,
 			SessionName,
 			OnSessionRestoredFn,
 			OnErrorFn));
 	};
 
-	auto OnArchiveIdLookupResultFn = [this, Promise, RepositoryId, ArchivedSessionId, Session, OnNewSessionRepositoryMountedFn, OnErrorFn](const FGuid& InArchivedSessionId)
+	auto OnSessionIdLookupResultFn = [this, Promise, RepositoryId, SrcSessionId, Session, OnNewSessionRepositoryMountedFn, OnErrorFn](const FGuid& InSrcSessionId)
 	{
 		if (RepositoryId != CurrentSessionRepositoryId) // Was aborted?
 		{
@@ -880,14 +886,14 @@ TFuture<TPair<bool, FText>> FDisasterRecoverySessionManager::RestoreAndJoinSessi
 			return;
 		}
 
-		if (!InArchivedSessionId.IsValid()) // Not found?
+		if (!InSrcSessionId.IsValid()) // Not found?
 		{
-			OnSessionNotFoundInternal(Session->RepositoryId); // The archive was likely deleted. Mark it to be discarded at the next start up.
-			OnErrorFn(LOCTEXT("ArchivedSessionNotFound", "Archived session not found."));
+			OnSessionNotFoundInternal(Session->RepositoryId); // The session was likely deleted. Mark its repository to be discarded at the next start up.
+			OnErrorFn(LOCTEXT("SessionNotFound", "Session not found."));
 			return;
 		}
 
-		*ArchivedSessionId = InArchivedSessionId; // Store the session ID for the next task above.
+		*SrcSessionId = InSrcSessionId; // Store the session ID for the next task above.
 
 		// 3- Mount a new repository to contain the new session (a new session is created from the one to restore).
 		TaskExecutor.Enqueue(MakeShared<FMountDisasterRecoverySessionRepositoryTask>(
@@ -896,9 +902,9 @@ TFuture<TPair<bool, FText>> FDisasterRecoverySessionManager::RestoreAndJoinSessi
 			OnErrorFn));
 	};
 
-	auto OnArchiveRepositoryMountedFn = [this, Promise, RepositoryId, Session, OnErrorFn, OnArchiveIdLookupResultFn]()
+	auto OnSrcRepositoryMountedFn = [this, Promise, RepositoryId, Session, OnErrorFn, OnSessionIdLookupResultFn]()
 	{
-		OnSessionRepositoryMountedInternal(Session->RepositoryId); // The archived session repository was mounted. Update internal state.
+		OnSessionRepositoryMountedInternal(Session->RepositoryId); // The session repository was mounted. Update internal state.
 
 		if (RepositoryId != CurrentSessionRepositoryId) // Was aborted?
 		{
@@ -906,18 +912,18 @@ TFuture<TPair<bool, FText>> FDisasterRecoverySessionManager::RestoreAndJoinSessi
 			return;
 		}
 
-		// 2- Search the archived session by name and retreive its ID.
+		// 2- Search the session by name and retreive its ID.
 		TaskExecutor.Enqueue(MakeShared<FLookupDisasterRecoverySessionIdTask>(
 			[this](){ return SyncClient->GetConcertClient()->GetServerSessions(ServerAdminEndpointId); },
 			Session->SessionName,
-			OnArchiveIdLookupResultFn,
+			OnSessionIdLookupResultFn,
 			OnErrorFn));
 	};
 
-	// 1- Mount the repository containing the archived session, the server will reload the archived session.
+	// 1- Mount the repository containing the session to restore, the server will reload the sessions in that repository.
 	TaskExecutor.Enqueue(MakeShared<FMountDisasterRecoverySessionRepositoryTask>(
 		[this, Session]() { return SyncClient->GetConcertClient()->MountSessionRepository(ServerAdminEndpointId, Session->RepositoryRootDir, Session->RepositoryId, /*bCreateIfNotExist*/false, /*bAsDefault*/false); },
-		OnArchiveRepositoryMountedFn,
+		OnSrcRepositoryMountedFn,
 		OnErrorFn));
 
 	return Promise->GetFuture();
diff --git a/Engine/Plugins/Developer/Concert/DisasterRecoveryClient/Source/DisasterRecoveryClient/Private/DisasterRecoveryTasks.cpp b/Engine/Plugins/Developer/Concert/DisasterRecoveryClient/Source/DisasterRecoveryClient/Private/DisasterRecoveryTasks.cpp
--- a/Engine/Plugins/Developer/Concert/DisasterRecoveryClient/Source/DisasterRecoveryClient/Private/DisasterRecoveryTasks.cpp
+++ b/Engine/Plugins/Developer/Concert/DisasterRecoveryClient/Source/DisasterRecoveryClient/Private/DisasterRecoveryTasks.cpp
@@ -391,7 +391,7 @@ bool FFindRecoverableDisasterRecoverySessionTask::Tick()
 				SearchingSessionIdTask = MakeUnique<FLookupDisasterRecoverySessionIdTask>(
 					[this](){ return Client->GetServerSessions(GetServerAdminEndpointIdFn()); },
 					Candidates.Last()->SessionName,
-					[this, Token](const FGuid& ArchivedSessionId) { if (Token.IsValid()) { SessionId = ArchivedSessionId; } }, // Search result. The ID might be valid (found) or invalid (not found).
+					[this, Token](const FGuid& ResultSessionId) { if (Token.IsValid()) { SessionId = ResultSessionId; } }, // Search result. The ID might be valid (found) or invalid (not found).
 					[this, Token](const FText&) { if (Token.IsValid()) { SessionId.Invalidate(); } }); // Error == not found, report an invalid GUID.
 				
 				SearchingSessionIdTask->Start(); // Transit to the 'searching session' state.
diff --git a/Engine/Plugins/Developer/Concert/MultiUserClient/Source/MultiUserClient/Private/Widgets/SConcertBrowser.cpp b/Engine/Plugins/Developer/Concert/MultiUserClient/Source/MultiUserClient/Private/Widgets/SConcertBrowser.cpp
--- a/Engine/Plugins/Developer/Concert/MultiUserClient/Source/MultiUserClient/Private/Widgets/SConcertBrowser.cpp
+++ b/Engine/Plugins/Developer/Concert/MultiUserClient/Source/MultiUserClient/Private/Widgets/SConcertBrowser.cpp
@@ -465,7 +465,7 @@ void FConcertBrowserController::RestoreSession(const FGuid& ServerAdminEndpointI
 	TWeakPtr<IConcertClient, ESPMode::ThreadSafe> WeakClient = ConcertClient;
 	auto OnAcceptRestoreFn = [WeakClient, ServerAdminEndpointId, SessionId, RestoredName, SessionFilter](TSharedPtr<FConcertClientSessionActivity> SelectedRecoveryActivity)
 	{
-		FConcertRestoreSessionArgs RestoreSessionArgs;
+		FConcertCopySessionArgs RestoreSessionArgs;
 		RestoreSessionArgs.bAutoConnect = true;
 		RestoreSessionArgs.SessionId = SessionId;
 		RestoreSessionArgs.SessionName = RestoredName;
diff --git a/Engine/Source/Programs/CrashReportClient/Private/CrashReportClientApp.cpp b/Engine/Source/Programs/CrashReportClient/Private/CrashReportClientApp.cpp
--- a/Engine/Source/Programs/CrashReportClient/Private/CrashReportClientApp.cpp
+++ b/Engine/Source/Programs/CrashReportClient/Private/CrashReportClientApp.cpp
@@ -834,18 +834,17 @@ void RunCrashReportClient(const TCHAR* CommandLine)
 
 					// Build error report in memory.
 					FPlatformErrorReport ErrorReport = CollectErrorReport(RecoveryServicePtr.Get(), MonitorPid, CrashContext, MonitorWritePipe);
-					const bool bNoDialog = (CrashContext.UserSettings.bNoDialog || CrashContext.UserSettings.bImplicitSend) && CrashContext.UserSettings.bSendUnattendedBugReports;
-					const SubmitCrashReportResult Result = SendErrorReport(ErrorReport, bNoDialog, CrashContext.UserSettings.bImplicitSend);
 
 #if CRASH_REPORT_WITH_RECOVERY
-					// Shut down the recovery service AFTER the UI prompt (if any) because archiving a large session can take several minutes (must be on main thread - UObject stuff). This ensures
-					// the UI is shown in a timely manner, but delays the server shutdown. At next reboot, user may not be able to restore because the session will still be archiving.
 					if (RecoveryServicePtr && !FPrimaryCrashProperties::Get()->bIsEnsure)
 					{
 						// Shutdown the recovery service. This will releases the recovery database file lock (not sharable) and let a new instance take it and offer the user to recover.
 						RecoveryServicePtr.Reset();
 					}
 #endif
+					const bool bNoDialog = (CrashContext.UserSettings.bNoDialog || CrashContext.UserSettings.bImplicitSend) && CrashContext.UserSettings.bSendUnattendedBugReports;
+					const SubmitCrashReportResult Result = SendErrorReport(ErrorReport, bNoDialog, CrashContext.UserSettings.bImplicitSend);
+
 					if (bReportCrashAnalyticInfo)
 					{
 						if (FCrashReportAnalytics::IsAvailable())
diff --git a/Engine/Source/Programs/CrashReportClient/Private/RecoveryService.cpp b/Engine/Source/Programs/CrashReportClient/Private/RecoveryService.cpp
--- a/Engine/Source/Programs/CrashReportClient/Private/RecoveryService.cpp
+++ b/Engine/Source/Programs/CrashReportClient/Private/RecoveryService.cpp
@@ -83,8 +83,9 @@ bool FRecoveryService::Startup()
 
 	// Setup the disaster recovery server configuration
 	UConcertServerConfig* ServerConfig = IConcertSyncServerModule::Get().ParseServerSettings(FCommandLine::Get());
-	ServerConfig->bAutoArchiveOnReboot = true; // If server crashed, was killed, etc, ensure the recovery session is archived (expected by recovery flow).
-	ServerConfig->EndpointSettings.RemoteEndpointTimeoutSeconds = 4; // Non-zero -> Ensure the Endpoints are re-registered with Message Bus if UDP layer restarts or some errors trigger a 'auto-repair'.
+	ServerConfig->bAutoArchiveOnReboot = false; // Skip archiving, disaster recovery restore a live session by copying it, this saves the step of archiving.
+	ServerConfig->bAutoArchiveOnShutdown = false; // Skip archiving, this can takes several minutes. It is more efficient to let the session 'live' and 'copy' it when restoring.
+	ServerConfig->EndpointSettings.RemoteEndpointTimeoutSeconds = 0; // Ensure the endpoints never time out (and are kept alive automatically by Concert).
 	ServerConfig->bMountDefaultSessionRepository = false; // Let the client mount its own repository to support concurrent recovery server and prevent them from concurrently accessing non-sharable database files.
 	ServerConfig->AuthorizedClientKeys.Add(ServerConfig->ServerName); // The disaster recovery client is configured to use the unique server name as key to identify itself.
 
diff --git a/Engine/Source/Programs/UnrealRecoverySvc/Private/UnrealRecoverySvc.cpp b/Engine/Source/Programs/UnrealRecoverySvc/Private/UnrealRecoverySvc.cpp
--- a/Engine/Source/Programs/UnrealRecoverySvc/Private/UnrealRecoverySvc.cpp
+++ b/Engine/Source/Programs/UnrealRecoverySvc/Private/UnrealRecoverySvc.cpp
@@ -39,8 +39,9 @@ INT32_MAIN_INT32_ARGC_TCHAR_ARGV()
 		}
 
 		UConcertServerConfig* ServerConfig = IConcertSyncServerModule::Get().ParseServerSettings(FCommandLine::Get());
-		ServerConfig->bAutoArchiveOnReboot = true; // If server crashed, was killed, etc, ensure the recovery session is archived (expected by recovery flow).
-		ServerConfig->EndpointSettings.RemoteEndpointTimeoutSeconds = 4; // Non-zero -> Ensure the Endpoints are re-registered with Message Bus if the UDP layer restarts or some errors trigger a 'auto-repair'.
+		ServerConfig->bAutoArchiveOnReboot = false; // Skip archiving, this can takes several minutes. It is more efficient to leave the live session as it and 'copy' it on demand (to restore).
+		ServerConfig->bAutoArchiveOnShutdown = false; // Skip archiving, this can takes several minutes. It is more efficient to leave the live session as it and 'copy' it on demand (to restore).
+		ServerConfig->EndpointSettings.RemoteEndpointTimeoutSeconds = 0; // Ensure the endpoints never time out (and are kept alive automatically by Concert).
 		ServerConfig->bMountDefaultSessionRepository = false; // Let the client mount its own repository to support concurrent service and avoid them to access the same non-sharable database files.
 		ServerConfig->AuthorizedClientKeys.Add(ServerConfig->ServerName); // The disaster recovery client is configured to use the unique server name as key to identify itself.
 		return ServerConfig;
