Changed the 'Disaster Recovery' friendly plugin name to 'Recovery Hub'.

diff --git a/Engine/Plugins/Developer/Concert/DisasterRecoveryClient/DisasterRecoveryClient.uplugin b/Engine/Plugins/Developer/Concert/DisasterRecoveryClient/DisasterRecoveryClient.uplugin
--- a/Engine/Plugins/Developer/Concert/DisasterRecoveryClient/DisasterRecoveryClient.uplugin
+++ b/Engine/Plugins/Developer/Concert/DisasterRecoveryClient/DisasterRecoveryClient.uplugin
@@ -2,7 +2,7 @@
 	"FileVersion": 3,
 	"Version": 1,
 	"VersionName": "0.2",
-	"FriendlyName": "Disaster Recovery",
+	"FriendlyName": "Recovery Hub",
 	"Description": "Track changes in the Editor to allow recovery in the event of a crash",
 	"Category": "Editor",
 	"CreatedBy": "Epic Games, Inc.",
