Fixed UI to display Recovery Hub rather than Disaster Recovery.
Limit the number of recovery session kept in the history between 0 to 50.

diff --git a/Engine/Plugins/Developer/Concert/DisasterRecoveryClient/Source/DisasterRecoveryClient/Private/DisasterRecoveryClientModule.cpp b/Engine/Plugins/Developer/Concert/DisasterRecoveryClient/Source/DisasterRecoveryClient/Private/DisasterRecoveryClientModule.cpp
--- a/Engine/Plugins/Developer/Concert/DisasterRecoveryClient/Source/DisasterRecoveryClient/Private/DisasterRecoveryClientModule.cpp
+++ b/Engine/Plugins/Developer/Concert/DisasterRecoveryClient/Source/DisasterRecoveryClient/Private/DisasterRecoveryClientModule.cpp
@@ -80,7 +80,7 @@ public:
 	{
 		if (!FApp::HasProjectName())
 		{
-			UE_LOG(LogDisasterRecovery, Warning, TEXT("The current project doesn't have any name. Disaster Recovery will be disabled! Please set your project name."));
+			UE_LOG(LogDisasterRecovery, Warning, TEXT("The current project doesn't have any name. Recovery Hub will be disabled! Please set your project name."));
 			return; // Needs a project name.
 		}
 
@@ -166,8 +166,8 @@ private:
 		if (ISettingsModule* SettingsModule = FModuleManager::GetModulePtr<ISettingsModule>("Settings"))
 		{
 			ISettingsSectionPtr SettingsSection = SettingsModule->RegisterSettings(DisasterRecoveryUtil::GetSettingsContainerName(), DisasterRecoveryUtil::GetSettingsCategoryName(), DisasterRecoveryUtil::GetSettingsSectionName(),
-				LOCTEXT("DisasterRecoverySettingsName", "Disaster Recovery"),
-				LOCTEXT("DisasterRecoverySettingsDescription", "Configure the Disaster Recovery Settings."),
+				LOCTEXT("DisasterRecoverySettingsName", "Recovery Hub"),
+				LOCTEXT("DisasterRecoverySettingsDescription", "Configure the Recovery Hub Settings."),
 				GetMutableDefault<UDisasterRecoverClientConfig>());
 
 			if (SettingsSection.IsValid())
@@ -217,7 +217,7 @@ private:
 			FOnSpawnTab::CreateRaw(this, &FDisasterRecoveryClientModule::SpawnRecoveryHub))
 			.SetIcon(FSlateIcon(FConcertFrontendStyle::GetStyleSetName(), TEXT("Concert.RecoveryHub")))
 			.SetDisplayName(LOCTEXT("RecoveryHubTabTitle", "Recovery Hub"))
-			.SetTooltipText(LOCTEXT("RecoveryHubTabTooltip", "Open the Disaster Recovery Hub"))
+			.SetTooltipText(LOCTEXT("RecoveryHubTabTooltip", "Open the Recovery Hub"))
 			.SetMenuType(ETabSpawnerMenuType::Enabled);
 
 		if (WorkspaceGroup.IsValid())
@@ -324,7 +324,7 @@ private:
 		const FString DisasterRecoveryServicePath = GetDisasterRecoveryServicePath();
 		if (DisasterRecoveryServicePath.IsEmpty())
 		{
-			UE_LOG(LogDisasterRecovery, Warning, TEXT("Disaster Recovery Service application was not found. Disaster Recovery will be disabled! Please build '%s'."), *DisasterRecoveryUtil::GetDisasterRecoveryServiceExeName());
+			UE_LOG(LogDisasterRecovery, Warning, TEXT("Unreal Recovery Service application was not found. Recovery Hub will be disabled! Please build '%s'."), *DisasterRecoveryUtil::GetDisasterRecoveryServiceExeName());
 			return false;
 		}
 
@@ -336,7 +336,7 @@ private:
 		DisasterRecoveryServiceHandle = FPlatformProcess::CreateProc(*DisasterRecoveryServicePath, *DisasterRecoveryServiceCommandLine, true, true, true, nullptr, 0, nullptr, nullptr, nullptr);
 		if (!DisasterRecoveryServiceHandle.IsValid())
 		{
-			UE_LOG(LogDisasterRecovery, Error, TEXT("Failed to launch Disaster Recovery Service application. Disaster Recovery will be disabled!"));
+			UE_LOG(LogDisasterRecovery, Error, TEXT("Failed to launch Recovery Service application. Recovery Hub will be disabled!"));
 			return false;
 		}
 
@@ -407,7 +407,7 @@ private:
 
 			// Launch the recovery hub UI.
 			TSharedRef<SWindow> NewWindow = SNew(SWindow)
-				.Title(LOCTEXT("RecoveryTitle", "Disaster Recovery"))
+				.Title(LOCTEXT("RecoveryTitle", "Recovery Hub"))
 				.SizingRule(ESizingRule::UserSized)
 				.ClientSize(FVector2D(1200, 800))
 				.IsTopmostWindow(true) // Keep it on top. Cannot be modal at the moment because Concert doesn't tick during a modal dialog.
@@ -439,7 +439,7 @@ private:
 	bool IsCompatibleWithOtherConcertSessions(const IConcertSyncClient* SyncClientStartingSession, const IConcertSyncClient* SyncClientShuttingDownSession) const
 	{
 		// At the moment, we don't expect more than 2 clients. We don't have use cases for a third concurrent concert client.
-		checkf(IConcertSyncClientModule::Get().GetClients().Num() <= 2, TEXT("Expected 1 disaster recovery client + 1 multi-user client at max."));
+		checkf(IConcertSyncClientModule::Get().GetClients().Num() <= 2, TEXT("Expected 1 recovery client + 1 multi-user client at max."));
 
 		// Scan all existing clients.
 		for (const TSharedRef<IConcertSyncClient>& SyncClient : IConcertSyncClientModule::Get().GetClients())
@@ -467,7 +467,7 @@ private:
 	bool IsCompatibleWithConcertClient(const IConcertSyncClient* SyncClient) const
 	{
 		check(SyncClient != DisasterRecoveryClient.Get());
-		checkf(SyncClient->GetConcertClient()->GetRole() == TEXT("MultiUser"), TEXT("A new role was added, check if this role can run concurrently with disaster recovery."));
+		checkf(SyncClient->GetConcertClient()->GetRole() == TEXT("MultiUser"), TEXT("A new Concert role was added, check if this role can run concurrently with recovery service."));
 
 		// Multi-User (MU) sessions are not compatible with disaster recovery (DR) session because MU events are performed in a transient sandbox that doesn't exist outside the MU session.
 		// If a crash occurs during a MU session, DR must not recover transactions applied to the transient sandbox. DR will will record the MU events, but for crash inspection purpose only.
diff --git a/Engine/Plugins/Developer/Concert/DisasterRecoveryClient/Source/DisasterRecoveryClient/Private/DisasterRecoverySettings.h b/Engine/Plugins/Developer/Concert/DisasterRecoveryClient/Source/DisasterRecoveryClient/Private/DisasterRecoverySettings.h
--- a/Engine/Plugins/Developer/Concert/DisasterRecoveryClient/Source/DisasterRecoveryClient/Private/DisasterRecoverySettings.h
+++ b/Engine/Plugins/Developer/Concert/DisasterRecoveryClient/Source/DisasterRecoveryClient/Private/DisasterRecoverySettings.h
@@ -15,7 +15,7 @@ public:
 	UDisasterRecoverClientConfig();
 
 	/**
-	 * Enables disaster recovery plugin to create and/or restore a recovery sessions when needed.
+	 * Enables Recovery Hub plugin to create and/or restore a recovery sessions when needed.
 	 */
 	UPROPERTY(config, EditAnywhere, Category="Client Settings")
 	bool bIsEnabled = true;
@@ -32,13 +32,13 @@ public:
 	 * The maximum number of recent recovery sessions to keep around. The sessions are rotated
 	 * at Editor startup and oldest ones are discarded.
 	 */
-	UPROPERTY(config, EditAnywhere, Category="Client Settings")
+	UPROPERTY(config, EditAnywhere, Category="Client Settings",  meta=(DisplayName="Session History Size", ClampMin = "0", UIMin = "0", UIMax = "50"))
 	int32 RecentSessionMaxCount = 4;
 
 	/**
 	 * The maximum number of imported recovery session to keep around. The sessions are rotated
 	 * at Editor startup and oldest imports are discarded.
 	 */
-	UPROPERTY(config, EditAnywhere, Category="Client Settings")
+	UPROPERTY(config, EditAnywhere, Category="Client Settings", meta=(DisplayName="Imported Session History Size", ClampMin = "0", UIMin = "0", UIMax = "50"))
 	int32 ImportedSessionMaxCount = 4;
 };
