Fixed missing WAL(write ahead log) file when importing Multi-User session into Disaster Recovery viewer (to debug the MU session).

diff --git a/Engine/Plugins/Developer/Concert/DisasterRecoveryClient/Source/DisasterRecoveryClient/Private/DisasterRecoverySessionManager.cpp b/Engine/Plugins/Developer/Concert/DisasterRecoveryClient/Source/DisasterRecoveryClient/Private/DisasterRecoverySessionManager.cpp
--- a/Engine/Plugins/Developer/Concert/DisasterRecoveryClient/Source/DisasterRecoveryClient/Private/DisasterRecoverySessionManager.cpp
+++ b/Engine/Plugins/Developer/Concert/DisasterRecoveryClient/Source/DisasterRecoveryClient/Private/DisasterRecoverySessionManager.cpp
@@ -606,6 +606,7 @@ TVariant<TSharedPtr<FDisasterRecoverySession>, FText> FDisasterRecoverySessionMa
 {
 	static const FString SessionInfoFilename = TEXT("SessionInfo.json");
 	static const FString SessionDatabaseFilename = TEXT("Session.db");
+	static const FString SessionDatabaseWalFilename = TEXT("Session.db-wal"); // The Write-ahead-log file, if present, likely contains some DB transactions.
 	//static const FString TransactionsDirName = TEXT("Transactions");
 	//static const FString PackagesDirName = TEXT("Packages");
 
@@ -624,6 +625,7 @@ TVariant<TSharedPtr<FDisasterRecoverySession>, FText> FDisasterRecoverySessionMa
 
 	// Copy the files in the archived session dir.
 	FString SessionDatabaseSrcPathname = FPaths::GetPath(SessionInfoPathname) / SessionDatabaseFilename;
+	FString SessionDatabaseWalSrcPathname = FPaths::GetPath(SessionInfoPathname) / SessionDatabaseWalFilename; // This file is not always there, depending on the database state.
 
 	if (SessionInfo.SessionName.IsEmpty())
 	{
@@ -655,6 +657,7 @@ TVariant<TSharedPtr<FDisasterRecoverySession>, FText> FDisasterRecoverySessionMa
 			// Copy the files in the SessionDir.
 			FString SessionInfoDestPathname = SessionDir / SessionInfoFilename;
 			FString SessionDatabaseDestPathname = SessionDir / SessionDatabaseFilename;
+			FString SessionDatabaseWalDestPathname = SessionDir / SessionDatabaseWalFilename;
 			if (IFileManager::Get().Copy(*SessionInfoDestPathname, *SessionInfoPathname) != COPY_OK)
 			{
 				ErrorMsg = FText::Format(LOCTEXT("FailedToCopySessionInfo", "Failed to copy '{0}' to '{1}'."), FText::AsCultureInvariant(SessionInfoPathname), FText::AsCultureInvariant(SessionDir));
@@ -665,6 +668,11 @@ TVariant<TSharedPtr<FDisasterRecoverySession>, FText> FDisasterRecoverySessionMa
 				ErrorMsg = FText::Format(LOCTEXT("FailedToCopySessionDB", "Failed to copy '{0}' to '{1}'."), FText::AsCultureInvariant(SessionDatabaseSrcPathname), FText::AsCultureInvariant(SessionDir));
 				ImportedSession.Reset();
 			}
+			else if (IFileManager::Get().FileExists(*SessionDatabaseWalSrcPathname) && IFileManager::Get().Copy(*SessionDatabaseWalDestPathname, *SessionDatabaseWalSrcPathname) != COPY_OK)
+			{
+				ErrorMsg = FText::Format(LOCTEXT("FailedToCopySessionDBWal", "Failed to copy '{0}' to '{1}'."), FText::AsCultureInvariant(SessionDatabaseWalSrcPathname), FText::AsCultureInvariant(SessionDir));
+				ImportedSession.Reset();
+			}
 			else
 			{
 				// TODO phase 2: Copy Transactions folder (if any)
