Added missing files

diff --git a/Engine/Plugins/Developer/Concert/DisasterRecoveryClient/DisasterRecoveryClient.uplugin b/Engine/Plugins/Developer/Concert/DisasterRecoveryClient/DisasterRecoveryClient.uplugin
--- a/Engine/Plugins/Developer/Concert/DisasterRecoveryClient/DisasterRecoveryClient.uplugin
+++ b/Engine/Plugins/Developer/Concert/DisasterRecoveryClient/DisasterRecoveryClient.uplugin
@@ -25,6 +25,10 @@
 		{
 			"Name": "ConcertSyncClient",
 			"Enabled": true
-		}
+		},
+        {
+          "Name": "ConcertSharedSlate",
+          "Enabled": true
+        }
 	]
 }
diff --git a/Engine/Plugins/Developer/Concert/DisasterRecoveryClient/Source/DisasterRecoveryClient/DisasterRecoveryClient.Build.cs b/Engine/Plugins/Developer/Concert/DisasterRecoveryClient/Source/DisasterRecoveryClient/DisasterRecoveryClient.Build.cs
--- a/Engine/Plugins/Developer/Concert/DisasterRecoveryClient/Source/DisasterRecoveryClient/DisasterRecoveryClient.Build.cs
+++ b/Engine/Plugins/Developer/Concert/DisasterRecoveryClient/Source/DisasterRecoveryClient/DisasterRecoveryClient.Build.cs
@@ -21,7 +21,7 @@ namespace UnrealBuildTool.Rules
 					"Concert",
 					"ConcertSyncClient",
 					"ConcertSyncCore",
-					"ConcertSyncUI",
+					"ConcertSharedSlate",
 					"ConcertTransport",
 					"DesktopPlatform",
 					"Json",
diff --git a/Engine/Plugins/Developer/Concert/MultiUserClient/MultiUserClient.uplugin b/Engine/Plugins/Developer/Concert/MultiUserClient/MultiUserClient.uplugin
--- a/Engine/Plugins/Developer/Concert/MultiUserClient/MultiUserClient.uplugin
+++ b/Engine/Plugins/Developer/Concert/MultiUserClient/MultiUserClient.uplugin
@@ -31,6 +31,10 @@
 			"Name": "ConcertSyncClient",
 			"Enabled": true
 		},
+        {
+          "Name": "ConcertSharedSlate",
+          "Enabled": true
+        },
 		{
 			"Name": "MultiUserTakes",
 			"Enabled": true
diff --git a/Engine/Plugins/Developer/Concert/MultiUserClient/Source/MultiUserClient/MultiUserClient.Build.cs b/Engine/Plugins/Developer/Concert/MultiUserClient/Source/MultiUserClient/MultiUserClient.Build.cs
--- a/Engine/Plugins/Developer/Concert/MultiUserClient/Source/MultiUserClient/MultiUserClient.Build.cs
+++ b/Engine/Plugins/Developer/Concert/MultiUserClient/Source/MultiUserClient/MultiUserClient.Build.cs
@@ -21,7 +21,7 @@ namespace UnrealBuildTool.Rules
 					"Concert",
 					"ConcertSyncClient",
 					"ConcertSyncCore",
-					"ConcertSyncUI",
+					"ConcertSharedSlate",
 					"ConcertTransport",
 					"DesktopPlatform",
 					"EditorStyle",
