Moving status bar styles to the core style location

diff --git a/Engine/Plugins/Developer/Concert/ConcertApp/DisasterRecoveryClient/Source/DisasterRecoveryClient/DisasterRecoveryClient.Build.cs b/Engine/Plugins/Developer/Concert/ConcertApp/DisasterRecoveryClient/Source/DisasterRecoveryClient/DisasterRecoveryClient.Build.cs
--- a/Engine/Plugins/Developer/Concert/ConcertApp/DisasterRecoveryClient/Source/DisasterRecoveryClient/DisasterRecoveryClient.Build.cs
+++ b/Engine/Plugins/Developer/Concert/ConcertApp/DisasterRecoveryClient/Source/DisasterRecoveryClient/DisasterRecoveryClient.Build.cs
@@ -32,6 +32,7 @@ namespace UnrealBuildTool.Rules
 					"Serialization",
 					"Slate",
 					"SlateCore",
+					"ToolWidgets",
 				}
 			);
 
diff --git a/Engine/Plugins/Developer/Concert/ConcertApp/DisasterRecoveryClient/Source/DisasterRecoveryClient/Private/SDisasterRecoveryHub.cpp b/Engine/Plugins/Developer/Concert/ConcertApp/DisasterRecoveryClient/Source/DisasterRecoveryClient/Private/SDisasterRecoveryHub.cpp
--- a/Engine/Plugins/Developer/Concert/ConcertApp/DisasterRecoveryClient/Source/DisasterRecoveryClient/Private/SDisasterRecoveryHub.cpp
+++ b/Engine/Plugins/Developer/Concert/ConcertApp/DisasterRecoveryClient/Source/DisasterRecoveryClient/Private/SDisasterRecoveryHub.cpp
@@ -26,6 +26,8 @@
 #include "Widgets/Layout/SUniformGridPanel.h"
 #include "Widgets/Input/SButton.h"
 #include "Widgets/Text/STextBlock.h"
+#include "SPositiveActionButton.h"
+#include "SNegativeActionButton.h"
 
 #define LOCTEXT_NAMESPACE "SDisasterRecoveryHub"
 
@@ -253,37 +255,21 @@ void SDisasterRecoveryHub::Construct(const FArguments& InArgs, const TSharedPtr<
 
 				+SUniformGridPanel::Slot(0, 0)
 				[
-					SNew(SButton)
-					.ForegroundColor(FLinearColor::White)
-					.ButtonStyle(FAppStyle::Get(), TEXT("FlatButton.Success"))
+					SNew(SPositiveActionButton)
+					.Icon(FAppStyle::Get().GetBrush("Icons.Refresh"))
 					.ToolTipText(this, &SDisasterRecoveryHub::GetRecoverAllButtonTooltip)
 					.OnClicked(this, &SDisasterRecoveryHub::OnRecoverAllButtonClicked)
-					.HAlign(HAlign_Center)
-					.ContentPadding(FMargin(14, 3))
 					.IsEnabled(this, &SDisasterRecoveryHub::IsRecoverAllButtonEnabled)
-					[
-						SNew(STextBlock)
-						.Text(LOCTEXT("RecoverAll", "Recover All"))
-						.Font( FAppStyle::Get().GetFontStyle("BoldFont"))
-						.ShadowOffset( FVector2D( 1.0f, 1.0f ) )
-					]
+					.Text(LOCTEXT("RecoverAll", "Recover All"))
 				]
 
 				+SUniformGridPanel::Slot(1, 0)
 				[
-					SNew(SButton)
-					.ForegroundColor(FLinearColor::White)
-					.ButtonStyle(FAppStyle::Get(), TEXT("FlatButton.Danger"))
+					SNew(SNegativeActionButton)
+					.Icon(FAppStyle::Get().GetBrush("Icons.Delete"))
 					.ToolTipText(LOCTEXT("CancelRecoveryTooltip", "Discard any recoverable data for your assets and continue with their last saved state"))
 					.OnClicked(this, &SDisasterRecoveryHub::OnCancelButtonClicked)
-					.HAlign(HAlign_Center)
-					.ContentPadding(FMargin(14, 3))
-					[
-						SNew(STextBlock)
-						.Text(LOCTEXT("Cancel", "Cancel"))
-						.Font(FAppStyle::Get().GetFontStyle("BoldFont"))
-						.ShadowOffset( FVector2D( 1.0f, 1.0f ) )
-					]
+					.Text(LOCTEXT("Discard", "Discard"))
 				]
 			]
 		]
diff --git a/Engine/Plugins/Developer/Concert/ConcertApp/MultiUserServer/Source/MultiUserServer/Private/Widgets/Clients/Logging/SConcertTransportLogFooter.cpp b/Engine/Plugins/Developer/Concert/ConcertApp/MultiUserServer/Source/MultiUserServer/Private/Widgets/Clients/Logging/SConcertTransportLogFooter.cpp
--- a/Engine/Plugins/Developer/Concert/ConcertApp/MultiUserServer/Source/MultiUserServer/Private/Widgets/Clients/Logging/SConcertTransportLogFooter.cpp
+++ b/Engine/Plugins/Developer/Concert/ConcertApp/MultiUserServer/Source/MultiUserServer/Private/Widgets/Clients/Logging/SConcertTransportLogFooter.cpp
@@ -111,12 +111,8 @@ void SConcertTransportLogFooter::Construct(const FArguments& InArgs, TSharedRef<
 		.VAlign(VAlign_Center)
 		[
 			SNew(SComboButton)
-				.ComboButtonStyle(FAppStyle::Get(), "GenericFilters.ComboButtonStyle")
-				.ForegroundColor(FLinearColor::White)
-				.ContentPadding(0)
 				.OnGetMenuContent(this, &SConcertTransportLogFooter::MakeViewOptionsMenuWidget, InArgs._ExtendViewOptions)
 				.HasDownArrow(true)
-				.ContentPadding(FMargin(1, 0))
 				.ButtonContent()
 				[
 					SNew(SHorizontalBox)
diff --git a/Engine/Plugins/Developer/Concert/ConcertUI/ConcertSharedSlate/Source/ConcertSharedSlate/Private/Session/Activity/SConcertSessionActivities.cpp b/Engine/Plugins/Developer/Concert/ConcertUI/ConcertSharedSlate/Source/ConcertSharedSlate/Private/Session/Activity/SConcertSessionActivities.cpp
--- a/Engine/Plugins/Developer/Concert/ConcertUI/ConcertSharedSlate/Source/ConcertSharedSlate/Private/Session/Activity/SConcertSessionActivities.cpp
+++ b/Engine/Plugins/Developer/Concert/ConcertUI/ConcertSharedSlate/Source/ConcertSharedSlate/Private/Session/Activity/SConcertSessionActivities.cpp
@@ -891,12 +891,8 @@ TSharedRef<SWidget> FConcertSessionActivitiesOptions::MakeDisplayedActivityCount
 TSharedRef<SWidget> FConcertSessionActivitiesOptions::MakeViewOptionsWidget(TOptional<FExtendContextMenu> ExtendMenu)
 {
 	return SNew(SComboButton)
-	.ComboButtonStyle(FAppStyle::Get(), "GenericFilters.ComboButtonStyle")
-	.ForegroundColor(FLinearColor::White)
-	.ContentPadding(0)
 	.OnGetMenuContent(this, &FConcertSessionActivitiesOptions::MakeMenuWidget, MoveTemp(ExtendMenu))
 	.HasDownArrow(true)
-	.ContentPadding(FMargin(1, 0))
 	.ButtonContent()
 	[
 		SNew(SHorizontalBox)
diff --git a/Engine/Source/Editor/EditorStyle/Private/StarshipStyle.cpp b/Engine/Source/Editor/EditorStyle/Private/StarshipStyle.cpp
--- a/Engine/Source/Editor/EditorStyle/Private/StarshipStyle.cpp
+++ b/Engine/Source/Editor/EditorStyle/Private/StarshipStyle.cpp
@@ -7485,30 +7485,7 @@
 {
 	Set("StatusBar.ContentBrowserUp", new CORE_IMAGE_BRUSH_SVG("Starship/Common/chevron-up", Icon16x16));
 	Set("StatusBar.ContentBrowserDown", new CORE_IMAGE_BRUSH_SVG("Starship/Common/chevron-down", Icon16x16));
-	Set("StatusBar.DrawerShadow", new BOX_BRUSH("Starship/StatusBar/drawer-shadow-bottom", FMargin(10/64.0f, 20/64.f, 10/64.f, 0), FLinearColor(0,0,0,1)));
-	Set("StatusBar.DrawerBackground", new FSlateRoundedBoxBrush(FStyleColors::Panel, 5.0f));
 
-	Set("StatusBar.Background", new FSlateColorBrush(FStyleColors::Panel));
-	Set("StatusBar.ProgressOverlay", new FSlateRoundedBoxBrush(FStyleColors::Transparent, 2.0f, FStyleColors::Panel, 1.0f, FVector2D(20, 8)));
-
-	Set("StatusBar.HelpIcon", new CORE_IMAGE_BRUSH_SVG("Starship/Common/help", Icon16x16, FStyleColors::Foreground));
-
-	FToolBarStyle StatusBarToolBarStyle = FStarshipCoreStyle::GetCoreStyle().GetWidgetStyle<FToolBarStyle>("SlimToolBar");
-
-	StatusBarToolBarStyle.SetBackground(FSlateNoResource());
-	StatusBarToolBarStyle.SetLabelPadding(FMargin(5, 5, 0, 5));
-	StatusBarToolBarStyle.SetIconSize(Icon16x16);
-	StatusBarToolBarStyle.SetBackgroundPadding(FMargin(4.f, 4.f, 0.f, 4.f));
-
-	Set("StatusBarToolBar", StatusBarToolBarStyle);
-
-	Set("StatusBar.Message.MessageText",
-		FTextBlockStyle(FAppStyle::Get().GetWidgetStyle<FTextBlockStyle>("NormalText"))
-		.SetColorAndOpacity(FStyleColors::Foreground));
-
-	Set("StatusBar.Message.HintText",
-		FTextBlockStyle(FAppStyle::Get().GetWidgetStyle<FTextBlockStyle>("NormalText"))
-		.SetColorAndOpacity(FStyleColors::Primary));
 }
 
 void FStarshipEditorStyle::FStyle::SetupColorPickerStyle()
diff --git a/Engine/Source/Runtime/SlateCore/Private/Styling/StarshipCoreStyle.cpp b/Engine/Source/Runtime/SlateCore/Private/Styling/StarshipCoreStyle.cpp
--- a/Engine/Source/Runtime/SlateCore/Private/Styling/StarshipCoreStyle.cpp
+++ b/Engine/Source/Runtime/SlateCore/Private/Styling/StarshipCoreStyle.cpp
@@ -1752,7 +1752,30 @@
 	Style->Set("StatusBar.StatusBarButton", StatusBarButton);
 	Style->Set("StatusBar.StatusBarComboButton", StatusBarComboButton);
 	Style->Set("StatusBar.Height", 32.0f);
+	Style->Set("StatusBar.DrawerShadow", new BOX_BRUSH("Starship/StatusBar/drawer-shadow-bottom", FMargin(10 / 64.0f, 20 / 64.f, 10 / 64.f, 0), FLinearColor(0, 0, 0, 1)));
+	Style->Set("StatusBar.DrawerBackground", new FSlateRoundedBoxBrush(FStyleColors::Panel, 5.0f));
 
+	Style->Set("StatusBar.Background", new FSlateColorBrush(FStyleColors::Panel));
+	Style->Set("StatusBar.ProgressOverlay", new FSlateRoundedBoxBrush(FStyleColors::Transparent, 2.0f, FStyleColors::Panel, 1.0f, FVector2D(20, 8)));
+
+	Style->Set("StatusBar.HelpIcon", new IMAGE_BRUSH_SVG("Starship/Common/help", Icon16x16, FStyleColors::Foreground));
+
+	FToolBarStyle StatusBarToolBarStyle = Style->GetWidgetStyle<FToolBarStyle>("SlimToolBar");
+
+	StatusBarToolBarStyle.SetBackground(FSlateNoResource());
+	StatusBarToolBarStyle.SetLabelPadding(FMargin(5, 5, 0, 5));
+	StatusBarToolBarStyle.SetIconSize(Icon16x16);
+	StatusBarToolBarStyle.SetBackgroundPadding(FMargin(4.f, 4.f, 0.f, 4.f));
+
+	Style->Set("StatusBarToolBar", StatusBarToolBarStyle);
+
+	Style->Set("StatusBar.Message.MessageText",
+		FTextBlockStyle(Style->GetWidgetStyle<FTextBlockStyle>("NormalText"))
+		.SetColorAndOpacity(FStyleColors::Foreground));
+
+	Style->Set("StatusBar.Message.HintText",
+		FTextBlockStyle(Style->GetWidgetStyle<FTextBlockStyle>("NormalText"))
+		.SetColorAndOpacity(FStyleColors::Primary));
 	// Create the regular border brush with all corners rounded and variants with certain corners squared.
 	{
 		const float CornerRadius = 5.0f;
